Curve Table improvements

UETOOL-281 - Improve data table workflows

Added support for exporting to CSV and JSON
Added support for importing from JSON (as well as the existing CSV support)
Ported over the new data table editor look-and-feel
Curve table editor now refreshes if the curve table it's editing is re-imported

diff --git a/Engine/Source/Developer/AssetTools/Private/AssetTypeActions/AssetTypeActions_CurveTable.cpp b/Engine/Source/Developer/AssetTools/Private/AssetTypeActions/AssetTypeActions_CurveTable.cpp
--- a/Engine/Source/Developer/AssetTools/Private/AssetTypeActions/AssetTypeActions_CurveTable.cpp
+++ b/Engine/Source/Developer/AssetTools/Private/AssetTypeActions/AssetTypeActions_CurveTable.cpp
@@ -5,6 +5,8 @@
 #include "Editor/CurveTableEditor/Public/CurveTableEditorModule.h"
 #include "Editor/CurveTableEditor/Public/ICurveTableEditor.h"
 #include "Engine/CurveTable.h"
+#include "DesktopPlatformModule.h"
+#include "IMainFrameModule.h"
 
 #define LOCTEXT_NAMESPACE "AssetTypeActions"
 
@@ -22,6 +24,26 @@ void FAssetTypeActions_CurveTable::GetActions( const TArray<UObject*>& InObjects
 		}
 	}
 
+	MenuBuilder.AddMenuEntry(
+		LOCTEXT("CurveTable_ExportAsCSV", "Export as CSV"),
+		LOCTEXT("CurveTable_ExportAsCSVTooltip", "Export the curve table as a file containing CSV data."),
+		FSlateIcon(),
+		FUIAction(
+			FExecuteAction::CreateSP( this, &FAssetTypeActions_CurveTable::ExecuteExportAsCSV, Tables ),
+			FCanExecuteAction()
+			)
+		);
+
+	MenuBuilder.AddMenuEntry(
+		LOCTEXT("CurveTable_ExportAsJSON", "Export as JSON"),
+		LOCTEXT("CurveTable_ExportAsJSONTooltip", "Export the curve table as a file containing JSON data."),
+		FSlateIcon(),
+		FUIAction(
+			FExecuteAction::CreateSP( this, &FAssetTypeActions_CurveTable::ExecuteExportAsJSON, Tables ),
+			FCanExecuteAction()
+			)
+		);
+
 	TArray<FString> PotentialFileExtensions;
 	PotentialFileExtensions.Add(TEXT(".xls"));
 	PotentialFileExtensions.Add(TEXT(".xlsm"));
@@ -29,7 +51,7 @@ void FAssetTypeActions_CurveTable::GetActions( const TArray<UObject*>& InObjects
 	PotentialFileExtensions.Add(TEXT(".json"));
 	MenuBuilder.AddMenuEntry(
 		LOCTEXT("CurveTable_OpenSourceData", "Open Source Data"),
-		LOCTEXT("CurveTable_OpenSourceDataTooltip", "Opens the data table's source data file in an external editor. It will search using the following extensions: .xls/.xlsm/.csv/.json"),
+		LOCTEXT("CurveTable_OpenSourceDataTooltip", "Opens the curve table's source data file in an external editor. It will search using the following extensions: .xls/.xlsm/.csv/.json"),
 		FSlateIcon(),
 		FUIAction(
 			FExecuteAction::CreateSP( this, &FAssetTypeActions_CurveTable::ExecuteFindSourceFileInExplorer, ImportPaths, PotentialFileExtensions ),
@@ -38,6 +60,88 @@ void FAssetTypeActions_CurveTable::GetActions( const TArray<UObject*>& InObjects
 		);
 }
 
+void FAssetTypeActions_CurveTable::ExecuteExportAsCSV(TArray< TWeakObjectPtr<UObject> > Objects)
+{
+	IDesktopPlatform* DesktopPlatform = FDesktopPlatformModule::Get();
+
+	void* ParentWindowWindowHandle = nullptr;
+
+	IMainFrameModule& MainFrameModule = FModuleManager::LoadModuleChecked<IMainFrameModule>(TEXT("MainFrame"));
+	const TSharedPtr<SWindow>& MainFrameParentWindow = MainFrameModule.GetParentWindow();
+	if ( MainFrameParentWindow.IsValid() && MainFrameParentWindow->GetNativeWindow().IsValid() )
+	{
+		ParentWindowWindowHandle = MainFrameParentWindow->GetNativeWindow()->GetOSWindowHandle();
+	}
+
+	for (auto ObjIt = Objects.CreateConstIterator(); ObjIt; ++ObjIt)
+	{
+		auto CurTable = Cast<UCurveTable>((*ObjIt).Get());
+		if (CurTable)
+		{
+			const FText Title = FText::Format(LOCTEXT("CurveTable_ExportCSVDialogTitle", "Export '{0}' as CSV..."), FText::FromString(*CurTable->GetName()));
+			const FString CurrentFilename = (CurTable->ImportPath.IsEmpty()) ? TEXT("") : FReimportManager::ResolveImportFilename(CurTable->ImportPath, CurTable);
+			const FString FileTypes = TEXT("Curve Table CSV (*.csv)|*.csv");
+
+			TArray<FString> OutFilenames;
+			DesktopPlatform->SaveFileDialog(
+				ParentWindowWindowHandle,
+				Title.ToString(),
+				(CurrentFilename.IsEmpty()) ? TEXT("") : FPaths::GetPath(CurrentFilename),
+				(CurrentFilename.IsEmpty()) ? TEXT("") : FPaths::GetBaseFilename(CurrentFilename) + TEXT(".csv"),
+				FileTypes,
+				EFileDialogFlags::None,
+				OutFilenames
+				);
+
+			if (OutFilenames.Num() > 0)
+			{
+				FFileHelper::SaveStringToFile(CurTable->GetTableAsCSV(), *OutFilenames[0]);
+			}
+		}
+	}
+}
+
+void FAssetTypeActions_CurveTable::ExecuteExportAsJSON(TArray< TWeakObjectPtr<UObject> > Objects)
+{
+	IDesktopPlatform* DesktopPlatform = FDesktopPlatformModule::Get();
+
+	void* ParentWindowWindowHandle = nullptr;
+
+	IMainFrameModule& MainFrameModule = FModuleManager::LoadModuleChecked<IMainFrameModule>(TEXT("MainFrame"));
+	const TSharedPtr<SWindow>& MainFrameParentWindow = MainFrameModule.GetParentWindow();
+	if ( MainFrameParentWindow.IsValid() && MainFrameParentWindow->GetNativeWindow().IsValid() )
+	{
+		ParentWindowWindowHandle = MainFrameParentWindow->GetNativeWindow()->GetOSWindowHandle();
+	}
+
+	for (auto ObjIt = Objects.CreateConstIterator(); ObjIt; ++ObjIt)
+	{
+		auto CurTable = Cast<UCurveTable>((*ObjIt).Get());
+		if (CurTable)
+		{
+			const FText Title = FText::Format(LOCTEXT("CurveTable_ExportCSVDialogTitle", "Export '{0}' as JSON..."), FText::FromString(*CurTable->GetName()));
+			const FString CurrentFilename = (CurTable->ImportPath.IsEmpty()) ? TEXT("") : FReimportManager::ResolveImportFilename(CurTable->ImportPath, CurTable);
+			const FString FileTypes = TEXT("Curve Table JSON (*.json)|*.json");
+
+			TArray<FString> OutFilenames;
+			DesktopPlatform->SaveFileDialog(
+				ParentWindowWindowHandle,
+				Title.ToString(),
+				(CurrentFilename.IsEmpty()) ? TEXT("") : FPaths::GetPath(CurrentFilename),
+				(CurrentFilename.IsEmpty()) ? TEXT("") : FPaths::GetBaseFilename(CurrentFilename) + TEXT(".json"),
+				FileTypes,
+				EFileDialogFlags::None,
+				OutFilenames
+				);
+
+			if (OutFilenames.Num() > 0)
+			{
+				FFileHelper::SaveStringToFile(CurTable->GetTableAsJSON(), *OutFilenames[0]);
+			}
+		}
+	}
+}
+
 void FAssetTypeActions_CurveTable::OpenAssetEditor( const TArray<UObject*>& InObjects, TSharedPtr<IToolkitHost> EditWithinLevelEditor )
 {
 	for (auto ObjIt = InObjects.CreateConstIterator(); ObjIt; ++ObjIt)
diff --git a/Engine/Source/Developer/AssetTools/Private/AssetTypeActions/AssetTypeActions_CurveTable.h b/Engine/Source/Developer/AssetTools/Private/AssetTypeActions/AssetTypeActions_CurveTable.h
--- a/Engine/Source/Developer/AssetTools/Private/AssetTypeActions/AssetTypeActions_CurveTable.h
+++ b/Engine/Source/Developer/AssetTools/Private/AssetTypeActions/AssetTypeActions_CurveTable.h
@@ -14,4 +14,11 @@ public:
 	virtual void OpenAssetEditor( const TArray<UObject*>& InObjects, TSharedPtr<class IToolkitHost> EditWithinLevelEditor = TSharedPtr<IToolkitHost>() ) override;
 	virtual void GetResolvedSourceFilePaths(const TArray<UObject*>& TypeAssets, TArray<FString>& OutSourceFilePaths) const override;
 	// End IAssetTypeActions
+
+private:
+	/** Handler for when CSV is selected */
+	void ExecuteExportAsCSV(TArray< TWeakObjectPtr<UObject> > Objects);
+
+	/** Handler for when JSON is selected */
+	void ExecuteExportAsJSON(TArray< TWeakObjectPtr<UObject> > Objects);
 };
\ No newline at end of file
diff --git a/Engine/Source/Editor/CurveTableEditor/CurveTableEditor.Build.cs b/Engine/Source/Editor/CurveTableEditor/CurveTableEditor.Build.cs
--- a/Engine/Source/Editor/CurveTableEditor/CurveTableEditor.Build.cs
+++ b/Engine/Source/Editor/CurveTableEditor/CurveTableEditor.Build.cs
@@ -13,6 +13,7 @@ public class CurveTableEditor : ModuleRules
 				"Core", 
 				"CoreUObject", 
 				"Engine", 
+				"InputCore",
 				"Slate",
 				"SlateCore",
                 "EditorStyle",
diff --git a/Engine/Source/Editor/CurveTableEditor/Private/CurveTableEditor.cpp b/Engine/Source/Editor/CurveTableEditor/Private/CurveTableEditor.cpp
--- a/Engine/Source/Editor/CurveTableEditor/Private/CurveTableEditor.cpp
+++ b/Engine/Source/Editor/CurveTableEditor/Private/CurveTableEditor.cpp
@@ -10,7 +10,47 @@
  
 #define LOCTEXT_NAMESPACE "CurveTableEditor"
 
-const FName FCurveTableEditor::CurveTableTabId( TEXT( "CurveTableEditor_CurveTable" ) );
+const FName FCurveTableEditor::CurveTableTabId("CurveTableEditor_CurveTable");
+const FName FCurveTableEditor::RowNameColumnId("RowName");
+
+class SCurveTableListViewRow : public SMultiColumnTableRow<FCurveTableEditorRowListViewDataPtr>
+{
+public:
+	SLATE_BEGIN_ARGS(SCurveTableListViewRow) {}
+		/** The widget that owns the tree.  We'll only keep a weak reference to it. */
+		SLATE_ARGUMENT(TSharedPtr<FCurveTableEditor>, CurveTableEditor)
+		/** The list item for this row */
+		SLATE_ARGUMENT(FCurveTableEditorRowListViewDataPtr, Item)
+	SLATE_END_ARGS()
+
+	/** Construct function for this widget */
+	void Construct(const FArguments& InArgs, const TSharedRef<STableViewBase>& InOwnerTableView)
+	{
+		CurveTableEditor = InArgs._CurveTableEditor;
+		Item = InArgs._Item;
+		SMultiColumnTableRow<FCurveTableEditorRowListViewDataPtr>::Construct(
+			FSuperRowType::FArguments()
+				.Style(FEditorStyle::Get(), "DataTableEditor.CellListViewRow"), 
+			InOwnerTableView
+			);
+	}
+
+	/** Overridden from SMultiColumnTableRow.  Generates a widget for this column of the list view. */
+	virtual TSharedRef<SWidget> GenerateWidgetForColumn(const FName& ColumnName) override
+	{
+		TSharedPtr<FCurveTableEditor> CurveTableEditorPtr = CurveTableEditor.Pin();
+		return (CurveTableEditorPtr.IsValid())
+			? CurveTableEditorPtr->MakeCellWidget(Item, IndexInList, ColumnName)
+			: SNullWidget::NullWidget;
+	}
+
+private:
+	/** Weak reference to the curve table editor that owns our list */
+	TWeakPtr<FCurveTableEditor> CurveTableEditor;
+	/** The item associated with this row of data */
+	FCurveTableEditorRowListViewDataPtr Item;
+};
+
 
 void FCurveTableEditor::RegisterTabSpawners(const TSharedRef<class FTabManager>& TabManager)
 {
@@ -21,14 +61,16 @@ void FCurveTableEditor::RegisterTabSpawners(const TSharedRef<class FTabManager>&
 		.SetGroup( WorkspaceMenuCategory.ToSharedRef() );
 }
 
+
 void FCurveTableEditor::UnregisterTabSpawners(const TSharedRef<class FTabManager>& TabManager)
 {
 	TabManager->UnregisterTabSpawner( CurveTableTabId );
 }
 
+
 FCurveTableEditor::~FCurveTableEditor()
 {
-
+	FReimportManager::Instance()->OnPostReimport().RemoveAll(this);
 }
 
 
@@ -52,6 +94,8 @@ void FCurveTableEditor::InitCurveTableEditor( const EToolkitMode::Type Mode, con
 	FCurveTableEditorModule& CurveTableEditorModule = FModuleManager::LoadModuleChecked<FCurveTableEditorModule>( "CurveTableEditor" );
 	AddMenuExtender(CurveTableEditorModule.GetMenuExtensibilityManager()->GetAllExtenders(GetToolkitCommands(), GetEditingObjects()));
 
+	FReimportManager::Instance()->OnPostReimport().AddSP(this, &FCurveTableEditor::OnPostReimport);
+
 	// @todo toolkit world centric editing
 	/*// Setup our tool's layout
 	if( IsWorldCentricAssetEditor() )
@@ -86,64 +130,47 @@ FLinearColor FCurveTableEditor::GetWorldCentricTabColorScale() const
 }
 
 
-TSharedRef<SUniformGridPanel> FCurveTableEditor::CreateGridPanel() const
+TSharedRef<SDockTab> FCurveTableEditor::SpawnTab_CurveTable( const FSpawnTabArgs& Args )
 {
-	TSharedRef<SUniformGridPanel> GridPanel = SNew(SUniformGridPanel).SlotPadding(FMargin(2.0f));
-
-	UCurveTable* Table = Cast<UCurveTable>(GetEditingObject());
-	if (Table != NULL)
-	{
-		int32 RowIdx = 0;
-		for ( auto RowIt = Table->RowMap.CreateConstIterator(); RowIt; ++RowIt )
-		{
-			const FName RowName = RowIt.Key();
-			const FRichCurve* RowCurve = RowIt.Value();
+	check( Args.GetTabId().TabType == CurveTableTabId );
 
-			// Make sure this row is valid
-			if ( RowCurve )
-			{
-				FLinearColor RowColor = (RowIdx % 2 == 0) ? FLinearColor::Gray : FLinearColor::Black;
-				int32 ColumnIdx = 0;
+	TSharedRef<SScrollBar> HorizontalScrollBar = SNew(SScrollBar)
+		.Orientation(Orient_Horizontal)
+		.Thickness(FVector2D(8.0f, 8.0f));
 
-				// Row name
-				GridPanel->AddSlot(ColumnIdx++, RowIdx)
-				[
-					SNew(SBorder)
-					.Padding(1)
-					.BorderImage(FEditorStyle::GetBrush("ToolPanel.GroupBorder"))
-					.BorderBackgroundColor(RowColor)
-					[
-						SNew(STextBlock) .Text(FText::FromName(RowName))
-					]
-				];
+	TSharedRef<SScrollBar> VerticalScrollBar = SNew(SScrollBar)
+		.Orientation(Orient_Vertical)
+		.Thickness(FVector2D(8.0f, 8.0f));
 
-				// Row Values
-				for (auto ValueIt(RowCurve->GetKeyIterator()); ValueIt; ++ValueIt)
-				{
-					GridPanel->AddSlot(ColumnIdx++, RowIdx)
-					[
-						SNew(SBorder)
-						.Padding(1)
-						.BorderImage(FEditorStyle::GetBrush("ToolPanel.GroupBorder"))
-						.BorderBackgroundColor(RowColor)
-						[
-							SNew(STextBlock).Text(FText::AsNumber(ValueIt->Value))
-						]
-					];
-				}
+	TSharedRef<SHeaderRow> RowNamesHeaderRow = SNew(SHeaderRow);
+	RowNamesHeaderRow->AddColumn(
+		SHeaderRow::Column(RowNameColumnId)
+		.DefaultLabel(FText::GetEmpty())
+		);
 
-				RowIdx++;
-			}
-		}
-	}
+	ColumnNamesHeaderRow = SNew(SHeaderRow);
 
-	return GridPanel;
-}
+	RowNamesListView = SNew(SListView<FCurveTableEditorRowListViewDataPtr>)
+		.ListItemsSource(&AvailableRows)
+		.HeaderRow(RowNamesHeaderRow)
+		.OnGenerateRow(this, &FCurveTableEditor::MakeRowNameWidget)
+		.OnListViewScrolled(this, &FCurveTableEditor::OnRowNamesListViewScrolled)
+		.ScrollbarVisibility(EVisibility::Collapsed)
+		.ConsumeMouseWheel(EConsumeMouseWheel::Always)
+		.SelectionMode(ESelectionMode::None)
+		.AllowOverscroll(EAllowOverscroll::No);
 
+	CellsListView = SNew(SListView<FCurveTableEditorRowListViewDataPtr>)
+		.ListItemsSource(&AvailableRows)
+		.HeaderRow(ColumnNamesHeaderRow)
+		.OnGenerateRow(this, &FCurveTableEditor::MakeRowWidget)
+		.OnListViewScrolled(this, &FCurveTableEditor::OnCellsListViewScrolled)
+		.ExternalScrollbar(VerticalScrollBar)
+		.ConsumeMouseWheel(EConsumeMouseWheel::Always)
+		.SelectionMode(ESelectionMode::None)
+		.AllowOverscroll(EAllowOverscroll::No);
 
-TSharedRef<SDockTab> FCurveTableEditor::SpawnTab_CurveTable( const FSpawnTabArgs& Args )
-{
-	check( Args.GetTabId().TabType == CurveTableTabId );
+	RefreshCachedCurveTable();
 
 	return SNew(SDockTab)
 		.Icon( FEditorStyle::GetBrush("CurveTableEditor.Tabs.Properties") )
@@ -154,18 +181,238 @@ TSharedRef<SDockTab> FCurveTableEditor::SpawnTab_CurveTable( const FSpawnTabArgs
 			.Padding(2)
 			.BorderImage(FEditorStyle::GetBrush("ToolPanel.GroupBorder"))
 			[
-				SNew(SScrollBox)
-				+ SScrollBox::Slot()
+				SNew(SVerticalBox)
+				+SVerticalBox::Slot()
 				[
-					SNew(SBox)
-					.VAlign(VAlign_Top)
-					.HAlign(HAlign_Left)
+					SNew(SHorizontalBox)
+					+SHorizontalBox::Slot()
+					.AutoWidth()
 					[
-						CreateGridPanel()
+						SNew(SBox)
+						.WidthOverride(this, &FCurveTableEditor::GetRowNameColumnWidth)
+						[
+							RowNamesListView.ToSharedRef()
+						]
+					]
+					+SHorizontalBox::Slot()
+					[
+						SNew(SScrollBox)
+						.Orientation(Orient_Horizontal)
+						.ExternalScrollbar(HorizontalScrollBar)
+						+SScrollBox::Slot()
+						[
+							CellsListView.ToSharedRef()
+						]
+					]
+					+SHorizontalBox::Slot()
+					.AutoWidth()
+					[
+						VerticalScrollBar
 					]
 				]
+				+SVerticalBox::Slot()
+				.AutoHeight()
+				[
+					SNew(SHorizontalBox)
+					+SHorizontalBox::Slot()
+					.AutoWidth()
+					[
+						SNew(SBox)
+						.WidthOverride(this, &FCurveTableEditor::GetRowNameColumnWidth)
+						[
+							SNullWidget::NullWidget
+						]
+					]
+					+SHorizontalBox::Slot()
+					[
+						HorizontalScrollBar
+					]
+				]
+			]
+		];
+}
+
+
+void FCurveTableEditor::RefreshCachedCurveTable()
+{
+	CacheDataTableForEditing();
+
+	ColumnNamesHeaderRow->ClearColumns();
+	for (int32 ColumnIndex = 0; ColumnIndex < AvailableColumns.Num(); ++ColumnIndex)
+	{
+		const FCurveTableEditorColumnHeaderDataPtr& ColumnData = AvailableColumns[ColumnIndex];
+
+		ColumnNamesHeaderRow->AddColumn(
+			SHeaderRow::Column(ColumnData->ColumnId)
+			.DefaultLabel(ColumnData->DisplayName)
+			.FixedWidth(ColumnData->DesiredColumnWidth)
+			);
+	}
+
+	RowNamesListView->RequestListRefresh();
+	CellsListView->RequestListRefresh();
+}
+
+
+void FCurveTableEditor::CacheDataTableForEditing()
+{
+	RowNameColumnWidth = 10.0f;
+
+	const UCurveTable* Table = Cast<const UCurveTable>(GetEditingObject());
+	if (!Table || Table->RowMap.Num() == 0)
+	{
+		AvailableColumns.Empty();
+		AvailableRows.Empty();
+		return;
+	}
+
+	TArray<FName> Names;
+	TArray<FRichCurve*> Curves;
+		
+	// get the row names and curves they represent
+	Table->RowMap.GenerateKeyArray(Names);
+	Table->RowMap.GenerateValueArray(Curves);
+
+	// Determine the curve with the longest set of data, for the column names
+	int32 LongestCurveIndex = 0;
+	for (int32 CurvesIdx = 1; CurvesIdx < Curves.Num(); ++CurvesIdx)
+	{
+		if (Curves[CurvesIdx]->GetNumKeys() > Curves[LongestCurveIndex]->GetNumKeys())
+		{
+			LongestCurveIndex = CurvesIdx;
+		}
+	}
+
+	TSharedRef<FSlateFontMeasure> FontMeasure = FSlateApplication::Get().GetRenderer()->GetFontMeasureService();
+	const FTextBlockStyle& CellTextStyle = FEditorStyle::GetWidgetStyle<FTextBlockStyle>("DataTableEditor.CellText");
+	static const float CellPadding = 10.0f;
+
+	// Column titles, taken from the longest curve
+	AvailableColumns.Empty(Curves[LongestCurveIndex]->GetNumKeys());
+	for (auto It(Curves[LongestCurveIndex]->GetKeyIterator()); It; ++It)
+	{
+		const FText ColumnText = FText::AsNumber(It->Time);
+
+		FCurveTableEditorColumnHeaderDataPtr CachedColumnData = MakeShareable(new FCurveTableEditorColumnHeaderData());
+		CachedColumnData->ColumnId = *ColumnText.ToString();
+		CachedColumnData->DisplayName = ColumnText;
+		CachedColumnData->DesiredColumnWidth = FontMeasure->Measure(CachedColumnData->DisplayName, CellTextStyle.Font).X + CellPadding;
+
+		AvailableColumns.Add(CachedColumnData);
+	}
+
+	// Each curve is a row entry
+	AvailableRows.Empty(Curves.Num());
+	for (int32 CurvesIdx = 0; CurvesIdx < Curves.Num(); ++CurvesIdx)
+	{
+		const FName& CurveName = Names[CurvesIdx];
+
+		FCurveTableEditorRowListViewDataPtr CachedRowData = MakeShareable(new FCurveTableEditorRowListViewData());
+		CachedRowData->RowId = CurveName;
+		CachedRowData->DisplayName = FText::FromName(CurveName);
+
+		const float RowNameWidth = FontMeasure->Measure(CachedRowData->DisplayName, CellTextStyle.Font).X + CellPadding;
+		RowNameColumnWidth = FMath::Max(RowNameColumnWidth, RowNameWidth);
+
+		CachedRowData->CellData.Reserve(Curves[CurvesIdx]->GetNumKeys());
+		{
+			int32 ColumnIndex = 0;
+			for (auto It(Curves[CurvesIdx]->GetKeyIterator()); It; ++It)
+			{
+				FCurveTableEditorColumnHeaderDataPtr CachedColumnData = AvailableColumns[ColumnIndex++];
+
+				const FText CellText = FText::AsNumber(It->Value);
+				CachedRowData->CellData.Add(CellText);
+
+				const float CellWidth = FontMeasure->Measure(CellText, CellTextStyle.Font).X + CellPadding;
+				CachedColumnData->DesiredColumnWidth = FMath::Max(CachedColumnData->DesiredColumnWidth, CellWidth);
+			}
+		}
+
+		AvailableRows.Add(CachedRowData);
+	}
+}
+
+
+TSharedRef<ITableRow> FCurveTableEditor::MakeRowNameWidget(FCurveTableEditorRowListViewDataPtr InRowDataPtr, const TSharedRef<STableViewBase>& OwnerTable)
+{
+	return
+		SNew(STableRow<FCurveTableEditorRowListViewDataPtr>, OwnerTable)
+		.Style(FEditorStyle::Get(), "DataTableEditor.NameListViewRow")
+		[
+			SNew(SBox)
+			.Padding(FMargin(4, 2, 4, 2))
+			[
+				SNew(STextBlock)
+				.Text(InRowDataPtr->DisplayName)
 			]
 		];
 }
 
+
+TSharedRef<ITableRow> FCurveTableEditor::MakeRowWidget(FCurveTableEditorRowListViewDataPtr InRowDataPtr, const TSharedRef<STableViewBase>& OwnerTable)
+{
+	return
+		SNew(SCurveTableListViewRow, OwnerTable)
+		.CurveTableEditor(SharedThis(this))
+		.Item(InRowDataPtr);
+}
+
+
+TSharedRef<SWidget> FCurveTableEditor::MakeCellWidget(FCurveTableEditorRowListViewDataPtr InRowDataPtr, const int32 InRowIndex, const FName& InColumnId)
+{
+	int32 ColumnIndex = 0;
+	for (; ColumnIndex < AvailableColumns.Num(); ++ColumnIndex)
+	{
+		const FCurveTableEditorColumnHeaderDataPtr& ColumnData = AvailableColumns[ColumnIndex];
+		if (ColumnData->ColumnId == InColumnId)
+		{
+			break;
+		}
+	}
+
+	// Valid column ID?
+	if (AvailableColumns.IsValidIndex(ColumnIndex))
+	{
+		return SNew(SBox)
+			.Padding(FMargin(4, 2, 4, 2))
+			[
+				SNew(STextBlock)
+				.TextStyle(FEditorStyle::Get(), "DataTableEditor.CellText")
+				.Text(InRowDataPtr->CellData[ColumnIndex])
+			];
+	}
+
+	return SNullWidget::NullWidget;
+}
+
+
+void FCurveTableEditor::OnRowNamesListViewScrolled(double InScrollOffset)
+{
+	// Synchronize the list views
+	CellsListView->SetScrollOffset(InScrollOffset);
+}
+
+
+void FCurveTableEditor::OnCellsListViewScrolled(double InScrollOffset)
+{
+	// Synchronize the list views
+	RowNamesListView->SetScrollOffset(InScrollOffset);
+}
+
+
+FOptionalSize FCurveTableEditor::GetRowNameColumnWidth() const
+{
+	return FOptionalSize(RowNameColumnWidth);
+}
+
+void FCurveTableEditor::OnPostReimport(UObject* InObject, bool)
+{
+	const UCurveTable* Table = Cast<const UCurveTable>(GetEditingObject());
+	if (Table && Table == InObject)
+	{
+		RefreshCachedCurveTable();
+	}
+}
+
 #undef LOCTEXT_NAMESPACE
diff --git a/Engine/Source/Editor/CurveTableEditor/Private/CurveTableEditor.h b/Engine/Source/Editor/CurveTableEditor/Private/CurveTableEditor.h
--- a/Engine/Source/Editor/CurveTableEditor/Private/CurveTableEditor.h
+++ b/Engine/Source/Editor/CurveTableEditor/Private/CurveTableEditor.h
@@ -7,10 +7,38 @@
 
 class UCurveTable;
 
+struct FCurveTableEditorColumnHeaderData
+{
+	/** Unique ID used to identify this column */
+	FName ColumnId;
+
+	/** Display name of this column */
+	FText DisplayName;
+
+	/** The calculated width of this column taking into account the cell data for each row */
+	float DesiredColumnWidth;
+};
+
+struct FCurveTableEditorRowListViewData
+{
+	/** Unique ID used to identify this row */
+	FName RowId;
+
+	/** Display name of this row */
+	FText DisplayName;
+
+	/** Array corresponding to each cell in this row */
+	TArray<FText> CellData;
+};
+
+typedef TSharedPtr<FCurveTableEditorColumnHeaderData> FCurveTableEditorColumnHeaderDataPtr;
+typedef TSharedPtr<FCurveTableEditorRowListViewData>  FCurveTableEditorRowListViewDataPtr;
+
 /** Viewer/editor for a CurveTable */
 class FCurveTableEditor :
 	public ICurveTableEditor
 {
+	friend class SCurveTableListViewRow;
 
 public:
 	virtual void RegisterTabSpawners(const TSharedRef<class FTabManager>& TabManager) override;
@@ -35,13 +63,61 @@ public:
 	virtual FLinearColor GetWorldCentricTabColorScale() const override;
 
 	/** Creates the grid widget used to display the data */
-	TSharedRef<SUniformGridPanel> CreateGridPanel() const;
+	//TSharedRef<SUniformGridPanel> CreateGridPanel() const;
 
 	/**	Spawns the tab with the curve table inside */
 	TSharedRef<SDockTab> SpawnTab_CurveTable( const FSpawnTabArgs& Args );
 
 private:
 
+	/** Update the cached state of this curve table, and then reflect that new state in the UI */
+	void RefreshCachedCurveTable();
+
+	/** Cache the data from the current curve table so that it can be shown in the editor */
+	void CacheDataTableForEditing();
+
+	/** Make the widget for a row name entry in the data table row list view */
+	TSharedRef<ITableRow> MakeRowNameWidget(FCurveTableEditorRowListViewDataPtr InRowDataPtr, const TSharedRef<STableViewBase>& OwnerTable);
+
+	/** Make the widget for a row entry in the data table row list view */
+	TSharedRef<ITableRow> MakeRowWidget(FCurveTableEditorRowListViewDataPtr InRowDataPtr, const TSharedRef<STableViewBase>& OwnerTable);
+
+	/** Make the widget for a cell entry in the data table row list view */
+	TSharedRef<SWidget> MakeCellWidget(FCurveTableEditorRowListViewDataPtr InRowDataPtr, const int32 InRowIndex, const FName& InColumnId);
+
+	/** Called when the row names list is scrolled - used to keep the two list views in sync */
+	void OnRowNamesListViewScrolled(double InScrollOffset);
+
+	/** Called when the cell names list view is scrolled - used to keep the two list views in sync */
+	void OnCellsListViewScrolled(double InScrollOffset);
+
+	/** Get the width to use for the row names column */
+	FOptionalSize GetRowNameColumnWidth() const;
+
+	/** Called when an asset has finished being imported */
+	void OnPostReimport(UObject* InObject, bool);
+
+	/** Array of the columns that are available for editing */
+	TArray<FCurveTableEditorColumnHeaderDataPtr> AvailableColumns;
+
+	/** Array of the rows that are available for editing */
+	TArray<FCurveTableEditorRowListViewDataPtr> AvailableRows;
+
+	/** Header row containing entries for each column in AvailableColumns */
+	TSharedPtr<SHeaderRow> ColumnNamesHeaderRow;
+
+	/** List view responsible for showing the row names column */
+	TSharedPtr<SListView<FCurveTableEditorRowListViewDataPtr>> RowNamesListView;
+
+	/** List view responsible for showing the rows from AvailableColumns */
+	TSharedPtr<SListView<FCurveTableEditorRowListViewDataPtr>> CellsListView;
+
+	/** Width of the row name column */
+	float RowNameColumnWidth;
+
 	/**	The tab id for the curve table tab */
 	static const FName CurveTableTabId;
-};
\ No newline at end of file
+
+	/** The column id for the row name list view column */
+	static const FName RowNameColumnId;
+};
diff --git a/Engine/Source/Editor/UnrealEd/Classes/Factories/ReimportDataTableFactory.h b/Engine/Source/Editor/UnrealEd/Classes/Factories/ReimportDataTableFactory.h
--- a/Engine/Source/Editor/UnrealEd/Classes/Factories/ReimportDataTableFactory.h
+++ b/Engine/Source/Editor/UnrealEd/Classes/Factories/ReimportDataTableFactory.h
@@ -14,10 +14,6 @@ class UReimportDataTableFactory : public UCSVImportFactory, public FReimportHand
 	virtual EReimportResult::Type Reimport( UObject* Obj ) override;
 	virtual int32 GetPriority() const override;
 	// End FReimportHandler interface
-
-	// Begin UCSVImportFactory interface
-	virtual TArray<FString> DoImportDataTable(class UDataTable* TargetDataTable, const FString& DataToImport) override;
-	// End UCSVImportFactory interface
 };
 
 
diff --git a/Engine/Source/Editor/UnrealEd/Private/Factories/CSVImportFactory.cpp b/Engine/Source/Editor/UnrealEd/Private/Factories/CSVImportFactory.cpp
--- a/Engine/Source/Editor/UnrealEd/Private/Factories/CSVImportFactory.cpp
+++ b/Engine/Source/Editor/UnrealEd/Private/Factories/CSVImportFactory.cpp
@@ -256,16 +256,39 @@ bool UCSVImportFactory::Reimport( UObject* Obj, const FString& Path )
 
 TArray<FString> UCSVImportFactory::DoImportDataTable(UDataTable* TargetDataTable, const FString& DataToImport)
 {
+	// Are we importing JSON data?
+	const bool bIsJSON = CurrentFilename.EndsWith(TEXT(".json"));
+	if (bIsJSON)
+	{
+		return TargetDataTable->CreateTableFromJSONString(DataToImport);
+	}
+
 	return TargetDataTable->CreateTableFromCSVString(DataToImport);
 }
 
 TArray<FString> UCSVImportFactory::DoImportCurveTable(UCurveTable* TargetCurveTable, const FString& DataToImport, const ERichCurveInterpMode ImportCurveInterpMode)
 {
+	// Are we importing JSON data?
+	const bool bIsJSON = CurrentFilename.EndsWith(TEXT(".json"));
+	if (bIsJSON)
+	{
+		return TargetCurveTable->CreateTableFromJSONString(DataToImport, ImportCurveInterpMode);
+	}
+
 	return TargetCurveTable->CreateTableFromCSVString(DataToImport, ImportCurveInterpMode);
 }
 
 TArray<FString> UCSVImportFactory::DoImportCurve(UCurveBase* TargetCurve, const FString& DataToImport)
 {
+	// Are we importing JSON data?
+	const bool bIsJSON = CurrentFilename.EndsWith(TEXT(".json"));
+	if (bIsJSON)
+	{
+		TArray<FString> Result;
+		Result.Add(LOCTEXT("Error_CannotImportCurveFromJSON", "Cannot import a curve from JSON. Please use CSV instead.").ToString());
+		return Result;
+	}
+
 	return TargetCurve->CreateCurveFromCSVString(DataToImport);
 }
 
@@ -314,24 +337,12 @@ int32 UReimportDataTableFactory::GetPriority() const
 	return ImportPriority;
 }
 
-TArray<FString> UReimportDataTableFactory::DoImportDataTable(UDataTable* TargetDataTable, const FString& DataToImport)
-{
-	// Are we importing JSON data?
-	const bool bIsJSON = CurrentFilename.EndsWith(TEXT(".json"));
-	if (bIsJSON)
-	{
-		return TargetDataTable->CreateTableFromJSONString(DataToImport);
-	}
-
-	// Fallback to CSV
-	return UCSVImportFactory::DoImportDataTable(TargetDataTable, DataToImport);
-}
-
 ////////////////////////////////////////////////////////////////////////////
 //
 UReimportCurveTableFactory::UReimportCurveTableFactory(const FObjectInitializer& ObjectInitializer)
 	: Super(ObjectInitializer)
 {
+	Formats.Add(TEXT("json;JavaScript Object Notation"));
 }
 
 bool UReimportCurveTableFactory::CanReimport( UObject* Obj, TArray<FString>& OutFilenames )
diff --git a/Engine/Source/Runtime/Engine/Classes/Engine/CurveTable.h b/Engine/Source/Runtime/Engine/Classes/Engine/CurveTable.h
--- a/Engine/Source/Runtime/Engine/Classes/Engine/CurveTable.h
+++ b/Engine/Source/Runtime/Engine/Classes/Engine/CurveTable.h
@@ -63,7 +63,10 @@ class UCurveTable
 	}
 
 	/** Output entire contents of table as a string */
-	ENGINE_API FString GetTableAsString();
+	ENGINE_API FString GetTableAsString() const;
+
+	/** Output entire contents of table as CSV */
+	ENGINE_API FString GetTableAsCSV() const;
 
 	/** Output entire contents of table as JSON */
 	ENGINE_API FString GetTableAsJSON() const;
@@ -80,6 +83,15 @@ class UCurveTable
 	 */
 	ENGINE_API TArray<FString> CreateTableFromCSVString(const FString& InString, ERichCurveInterpMode InterpMode = RCIM_Linear);
 
+	/** 
+	 *	Create table from JSON string. 
+	 *	RowCurve must be defined before calling this function. 
+	 *  @param InString The string representing the CurveTable
+	 *  @param InterpMode The mode of interpolation to use for the curves
+	 *	@return	Set of problems encountered while processing input
+	 */
+	ENGINE_API TArray<FString> CreateTableFromJSONString(const FString& InString, ERichCurveInterpMode InterpMode = RCIM_Linear);
+
 	/** Empty the table info (will not clear RowCurve) */
 	ENGINE_API void EmptyTable();
 
diff --git a/Engine/Source/Runtime/Engine/Private/CurveTable.cpp b/Engine/Source/Runtime/Engine/Private/CurveTable.cpp
--- a/Engine/Source/Runtime/Engine/Private/CurveTable.cpp
+++ b/Engine/Source/Runtime/Engine/Private/CurveTable.cpp
@@ -100,7 +100,7 @@ void UCurveTable::GetAssetRegistryTags(TArray<FAssetRegistryTag>& OutTags) const
 }
 #endif
 
-FString UCurveTable::GetTableAsString()
+FString UCurveTable::GetTableAsString() const
 {
 	FString Result;
 
@@ -153,6 +153,56 @@ FString UCurveTable::GetTableAsString()
 	return Result;
 }
 
+FString UCurveTable::GetTableAsCSV() const
+{
+	FString Result;
+
+	if(RowMap.Num() > 0)
+	{
+		TArray<FName> Names;
+		TArray<FRichCurve*> Curves;
+		
+		// get the row names and curves they represent
+		RowMap.GenerateKeyArray(Names);
+		RowMap.GenerateValueArray(Curves);
+
+		// Determine the curve with the longest set of data, for headers
+		int32 LongestCurveIndex = 0;
+		for (int32 CurvesIdx = 1; CurvesIdx < Curves.Num(); CurvesIdx++)
+		{
+			if(Curves[CurvesIdx]->GetNumKeys() > Curves[LongestCurveIndex]->GetNumKeys())
+			{
+				LongestCurveIndex = CurvesIdx;
+			}
+		}
+
+		// First row, column titles, taken from the longest curve
+		Result += TEXT("---");
+		for (auto It(Curves[LongestCurveIndex]->GetKeyIterator()); It; ++It)
+		{
+			Result += FString::Printf(TEXT(",%f"), It->Time);
+		}
+		Result += TEXT("\n");
+
+		// display all the curves
+		for (int32 CurvesIdx = 0; CurvesIdx < Curves.Num(); CurvesIdx++)
+		{
+			// show name of curve
+			Result += Names[CurvesIdx].ToString();
+
+			// show data of curve
+			for (auto It(Curves[CurvesIdx]->GetKeyIterator()); It; ++It)
+			{
+				Result += FString::Printf(TEXT(",%f"), It->Value);
+			}
+
+			Result += TEXT("\n");
+		}
+	}
+
+	return Result;
+}
+
 FString UCurveTable::GetTableAsJSON() const
 {
 	// use the pretty print policy since these values are usually getting dumpped for check-in to P4 (or for inspection)
@@ -314,6 +364,96 @@ TArray<FString> UCurveTable::CreateTableFromCSVString(const FString& InString, E
 	return OutProblems;
 }
 
+TArray<FString> UCurveTable::CreateTableFromJSONString(const FString& InString, ERichCurveInterpMode InterpMode)
+{
+	// Array used to store problems about table creation
+	TArray<FString> OutProblems;
+
+	if (InString.IsEmpty())
+	{
+		OutProblems.Add(TEXT("Input data is empty."));
+		return OutProblems;
+	}
+
+	TArray< TSharedPtr<FJsonValue> > ParsedTableRows;
+	{
+		const TSharedRef< TJsonReader<TCHAR> > JsonReader = TJsonReaderFactory<TCHAR>::Create(InString);
+		if (!FJsonSerializer::Deserialize(JsonReader, ParsedTableRows) || ParsedTableRows.Num() == 0)
+		{
+			OutProblems.Add(FString::Printf(TEXT("Failed to parse the JSON data. Error: %s"), *JsonReader->GetErrorMessage()));
+			return OutProblems;
+		}
+	}
+
+	// Empty existing data
+	EmptyTable();
+
+	// Iterate over rows
+	for (int32 RowIdx = 0; RowIdx < ParsedTableRows.Num(); ++RowIdx)
+	{
+		const TSharedPtr<FJsonValue>& ParsedTableRowValue = ParsedTableRows[RowIdx];
+		TSharedPtr<FJsonObject> ParsedTableRowObject = ParsedTableRowValue->AsObject();
+		if (!ParsedTableRowObject.IsValid())
+		{
+			OutProblems.Add(FString::Printf(TEXT("Row '%d' is not a valid JSON object."), RowIdx));
+			continue;
+		}
+
+		// Get row name
+		static const FString RowNameJsonKey = TEXT("Name");
+		const FName RowName = MakeValidName(ParsedTableRowObject->GetStringField(RowNameJsonKey));
+
+		// Check its not 'none'
+		if (RowName == NAME_None)
+		{
+			OutProblems.Add(FString::Printf(TEXT("Row '%d' missing a name."), RowIdx));
+			continue;
+		}
+
+		// Check its not a duplicate
+		if (RowMap.Find(RowName) != NULL)
+		{
+			OutProblems.Add(FString::Printf(TEXT("Duplicate row name '%s'."), *RowName.ToString()));
+			continue;
+		}
+
+		// Add a key for each entry in this row
+		FRichCurve* NewCurve = new FRichCurve();
+		for (const auto& ParsedTableRowEntry : ParsedTableRowObject->Values)
+		{
+			// Skip the name entry
+			if (ParsedTableRowEntry.Key == RowNameJsonKey)
+			{
+				continue;
+			}
+
+			// Make sure we have a valid float key
+			float EntryKey = 0.0f;
+			if (!LexicalConversion::TryParseString(EntryKey, *ParsedTableRowEntry.Key))
+			{
+				OutProblems.Add(FString::Printf(TEXT("Key '%s' on row '%s' is not a float and cannot be parsed."), *ParsedTableRowEntry.Key, *RowName.ToString()));
+				continue;
+			}
+
+			// Make sure we have a valid float value
+			double EntryValue = 0.0;
+			if (!ParsedTableRowEntry.Value->TryGetNumber(EntryValue))
+			{
+				OutProblems.Add(FString::Printf(TEXT("Entry '%s' on row '%s' is not a float and cannot be parsed."), *ParsedTableRowEntry.Key, *RowName.ToString()));
+				continue;
+			}
+
+			FKeyHandle KeyHandle = NewCurve->AddKey(EntryKey, static_cast<float>(EntryValue));
+			NewCurve->SetKeyInterpMode(KeyHandle, InterpMode);
+		}
+
+		RowMap.Add(RowName, NewCurve);
+	}
+
+	Modify(true);
+	return OutProblems;
+}
+
 //////////////////////////////////////////////////////////////////////////
 
 
