CurveTableEditor - work in progress

diff --git a/Engine/Source/Editor/CurveEditor/Private/Tree/SCurveEditorTree.cpp b/Engine/Source/Editor/CurveEditor/Private/Tree/SCurveEditorTree.cpp
--- a/Engine/Source/Editor/CurveEditor/Private/Tree/SCurveEditorTree.cpp
+++ b/Engine/Source/Editor/CurveEditor/Private/Tree/SCurveEditorTree.cpp
@@ -101,6 +101,7 @@ void SCurveEditorTree::Construct(const FArguments& InArgs, TSharedPtr<FCurveEdit
 		.OnGenerateRow(this, &SCurveEditorTree::GenerateRow)
 		.OnSetExpansionRecursive(this, &SCurveEditorTree::SetItemExpansionRecursive)
 		.OnMouseButtonDoubleClick(InArgs._OnMouseButtonDoubleClick)
+		.OnTreeViewScrolled(InArgs._OnTreeViewScrolled)
 		.OnSelectionChanged_Lambda(
 			[this](TListTypeTraits<FCurveEditorTreeItemID>::NullableType InItemID, ESelectInfo::Type Type)
 			{
diff --git a/Engine/Source/Editor/CurveEditor/Public/Tree/SCurveEditorTree.h b/Engine/Source/Editor/CurveEditor/Public/Tree/SCurveEditorTree.h
--- a/Engine/Source/Editor/CurveEditor/Public/Tree/SCurveEditorTree.h
+++ b/Engine/Source/Editor/CurveEditor/Public/Tree/SCurveEditorTree.h
@@ -14,6 +14,7 @@ class STableViewBase;
 
 class CURVEEDITOR_API SCurveEditorTree : public STreeView<FCurveEditorTreeItemID>
 {
+
 public:
 
 	SLATE_BEGIN_ARGS(SCurveEditorTree)
@@ -21,17 +22,21 @@ public:
 		{}
 		SLATE_ARGUMENT(float, SelectColumnWidth)
 		SLATE_EVENT(FOnMouseButtonDoubleClick, OnMouseButtonDoubleClick)
+		SLATE_EVENT(FOnTableViewScrolled, OnTreeViewScrolled)
 	SLATE_END_ARGS()
 
 	void Construct(const FArguments& InArgs, TSharedPtr<FCurveEditor> InCurveEditor);
 
+	const TArray<FCurveEditorTreeItemID>& GetSourceItems() const { return RootItems; }
+
+	TSharedRef<ITableRow> GenerateRow(FCurveEditorTreeItemID ItemID, const TSharedRef<STableViewBase>& OwnerTable);
+	
 private:
 
 	virtual FReply OnKeyDown(const FGeometry& MyGeometry, const FKeyEvent& InKeyEvent) override;
 
 private:
 
-	TSharedRef<ITableRow> GenerateRow(FCurveEditorTreeItemID ItemID, const TSharedRef<STableViewBase>& OwnerTable);
 
 	void GetTreeItemChildren(FCurveEditorTreeItemID Parent, TArray<FCurveEditorTreeItemID>& OutChildren);
 
@@ -60,4 +65,5 @@ private:
 
 	bool bUpdatingTreeWidgetSelection;
 	bool bUpdatingCurveEditorTreeSelection;
+
 };
\ No newline at end of file
diff --git a/Engine/Source/Editor/CurveTableEditor/CurveTableEditor.Build.cs b/Engine/Source/Editor/CurveTableEditor/CurveTableEditor.Build.cs
--- a/Engine/Source/Editor/CurveTableEditor/CurveTableEditor.Build.cs
+++ b/Engine/Source/Editor/CurveTableEditor/CurveTableEditor.Build.cs
@@ -12,6 +12,7 @@ public class CurveTableEditor : ModuleRules
 			new string[] {
 				"Core", 
 				"CoreUObject", 
+                "CurveEditor",
 				"EditorFramework",
 				"Engine", 
 				"InputCore",
diff --git a/Engine/Source/Editor/CurveTableEditor/Private/CurveTableEditor.cpp b/Engine/Source/Editor/CurveTableEditor/Private/CurveTableEditor.cpp
--- a/Engine/Source/Editor/CurveTableEditor/Private/CurveTableEditor.cpp
+++ b/Engine/Source/Editor/CurveTableEditor/Private/CurveTableEditor.cpp
@@ -7,122 +7,211 @@
 #include "Framework/Application/SlateApplication.h"
 #include "Widgets/Layout/SScrollBar.h"
 #include "Framework/Layout/Overscroll.h"
+#include "Widgets/Input/SNumericEntryBox.h"
+#include "Widgets/Input/SSegmentedControl.h"
 #include "Widgets/Views/SListView.h"
 #include "Widgets/Layout/SScrollBox.h"
+#include "SEditorHeaderButton.h"
 #include "EditorStyleSet.h"
+#include "Styling/StyleColors.h"
 #include "EditorReimportHandler.h"
 #include "CurveTableEditorModule.h"
 #include "Widgets/Docking/SDockTab.h"
-#include "SCurveEditor.h"
+#include "CurveEditor.h" 
+#include "SCurveEditorPanel.h"
+#include "Tree/SCurveEditorTree.h"
+#include "Tree/SCurveEditorTreeSelect.h"
+#include "Tree/SCurveEditorTreePin.h"
+#include "Tree/ICurveEditorTreeItem.h"
+#include "Tree/CurveEditorTreeFilter.h"
+
+#include "Tree/SCurveEditorTreeTextFilter.h"
+
+#include "RealCurveModel.h"
+#include "RichCurveEditorModel.h"
+
 #include "CurveTableEditorCommands.h"
 #include "Framework/MultiBox/MultiBoxBuilder.h"
+
  
 #define LOCTEXT_NAMESPACE "CurveTableEditor"
 
 const FName FCurveTableEditor::CurveTableTabId("CurveTableEditor_CurveTable");
-const FName FCurveTableEditor::RowNameColumnId("RowName");
 
-class SCurveTableEditor : public SCurveEditor
+struct FCurveTableEditorColumnHeaderData
 {
-	SLATE_BEGIN_ARGS(SCurveTableEditor) {}
-	SLATE_END_ARGS()
+	/** Unique ID used to identify this column */
+	FName ColumnId;
 
-	void Construct(const FArguments& InArgs)
-	{
-		SCurveEditor::Construct(SCurveEditor::FArguments()
-			.DesiredSize(FVector2D(128.0f, 64.0f)));
-	}
+	/** Display name of this column */
+	FText DisplayName;
 
-	virtual FReply OnMouseWheel(const FGeometry& MyGeometry, const FPointerEvent& MouseEvent) override
-	{
-		if (MouseEvent.IsShiftDown())
-		{
-			SCurveEditor::OnMouseWheel(MyGeometry, MouseEvent);
-			return FReply::Handled();
-		}
-		
-		return FReply::Unhandled();
-	}
+	/** The calculated width of this column taking into account the cell data for each row */
+	float DesiredColumnWidth;
+
+	/** The evaluated key time **/
+	float KeyTime;
 };
 
-class SCurveTableListViewRow : public SMultiColumnTableRow<FCurveTableEditorRowListViewDataPtr>
+
+/*
+* FCurveTableEditorItem
+*
+*  FCurveTableEditorItem uses and extends the CurveEditorTreeItem to be used in both our TableView and the CurveEditorTree.
+*  The added GenerateTableViewCell handles the table columns unknown to the standard CurveEditorTree.
+*
+*/ 
+class FCurveTableEditorItem : public ICurveEditorTreeItem,  public TSharedFromThis<FCurveTableEditorItem>
 {
-public:
-	SLATE_BEGIN_ARGS(SCurveTableListViewRow) {}
-		/** The widget that owns the tree.  We'll only keep a weak reference to it. */
-		SLATE_ARGUMENT(TSharedPtr<FCurveTableEditor>, CurveTableEditor)
-		/** The list item for this row */
-		SLATE_ARGUMENT(FCurveTableEditorRowListViewDataPtr, Item)
-	SLATE_END_ARGS()
 
-	/** Construct function for this widget */
-	void Construct(const FArguments& InArgs, const TSharedRef<STableViewBase>& InOwnerTableView)
+  	struct CachedKeyInfo
+  	{
+  		CachedKeyInfo(FKeyHandle& InKeyHandle, FText InDisplayValue) :
+  		KeyHandle(InKeyHandle)
+  		, DisplayValue(InDisplayValue) {}
+
+  		FKeyHandle KeyHandle;
+
+  		FText DisplayValue;	
+  	};
+
+  public: 
+	FCurveTableEditorItem (const FName& InRowId, FCurveTableEditorHandle InRowHandle, const TArray<FCurveTableEditorColumnHeaderDataPtr>& InColumns)
+		: RowId(InRowId)
+		, RowHandle(InRowHandle)
+		, Columns(InColumns)
 	{
-		CurveTableEditor = InArgs._CurveTableEditor;
-		Item = InArgs._Item;
-		SMultiColumnTableRow<FCurveTableEditorRowListViewDataPtr>::Construct(
-			FSuperRowType::FArguments()
-				.Style(FEditorStyle::Get(), "DataTableEditor.CellListViewRow"), 
-			InOwnerTableView
-			);
+		DisplayName = FText::FromName(InRowId);
+
+		CacheKeys();
 	}
 
-	virtual void ConstructChildren(ETableViewMode::Type InOwnerTableMode, const TAttribute<FMargin>& InPadding, const TSharedRef<SWidget>& InContent) override
+	TSharedPtr<SWidget> GenerateCurveEditorTreeWidget(const FName& InColumnName, TWeakPtr<FCurveEditor> InCurveEditor, FCurveEditorTreeItemID InTreeItemID, const TSharedRef<ITableRow>& InTableRow) override
 	{
-		TSharedPtr<FCurveTableEditor> CurveTableEditorPtr = CurveTableEditor.Pin();
-		if (CurveTableEditorPtr.IsValid())
+		if (InColumnName == ColumnNames.Label)
 		{
-			ChildSlot
-			[
-				SNew(SVerticalBox)
-				+SVerticalBox::Slot()
-				.FillHeight(1.0f)
-				[
-					SNew(SBox)
-					.Visibility(this, &SCurveTableListViewRow::GetTableViewVisibility)
-					[
-						InContent
-					]
-				]
-				+SVerticalBox::Slot()
-				.FillHeight(1.0f)
+			return SNew(SHorizontalBox)
+				+SHorizontalBox::Slot()
+				.Padding(FMargin(4.f))
+				.VAlign(VAlign_Center)
+				.HAlign(HAlign_Right)
+				.AutoWidth()
 				[
-					SNew(SHorizontalBox)
-					.Visibility(this, &SCurveTableListViewRow::GetCurveViewVisibility)
-					+SHorizontalBox::Slot()
-					.FillWidth(1.0f)
-					[
-						CurveTableEditorPtr->MakeCurveWidget(Item, IndexInList)
-					]
-				]
-			];
+					SNew(STextBlock)
+					.Text(DisplayName)
+					.ColorAndOpacity(FSlateColor::UseForeground())
+				];
+		}
+		else if (InColumnName == ColumnNames.SelectHeader)
+		{
+			return SNew(SCurveEditorTreeSelect, InCurveEditor, InTreeItemID, InTableRow);
+		}
+		else if (InColumnName == ColumnNames.PinHeader)
+		{
+			return SNew(SCurveEditorTreePin, InCurveEditor, InTreeItemID, InTableRow);
 		}
+
+		return GenerateTableViewCell(InColumnName, InCurveEditor, InTreeItemID, InTableRow);
 	}
 
-	/** Overridden from SMultiColumnTableRow.  Generates a widget for this column of the list view. */
-	virtual TSharedRef<SWidget> GenerateWidgetForColumn(const FName& ColumnName) override
+	TSharedPtr<SWidget> GenerateTableViewCell(const FName& InColumnId, TWeakPtr<FCurveEditor> InCurveEditor, FCurveEditorTreeItemID InTreeItemID, const TSharedRef<ITableRow>& InTableRow)
 	{
-		TSharedPtr<FCurveTableEditor> CurveTableEditorPtr = CurveTableEditor.Pin();
-		return (CurveTableEditorPtr.IsValid())
-			? CurveTableEditorPtr->MakeCellWidget(Item, IndexInList, ColumnName)
-			: SNullWidget::NullWidget;
+		if (!RowHandle.HasRichCurves())
+		{
+			FRealCurve* Curve = RowHandle.GetCurve();
+			FKeyHandle& KeyHandle = CellDataMap[InColumnId].KeyHandle;
+
+			return SNew(SNumericEntryBox<float>)
+				.EditableTextBoxStyle( &FAppStyle::Get().GetWidgetStyle<FEditableTextBoxStyle>("CurveTableEditor.Cell.Text") )
+				.Value_Lambda([Curve, KeyHandle] () { return Curve->GetKeyValue(KeyHandle); })
+				// .Value_Lambda([this, InColumnId] () { return CellDataF[InColumnId]; })
+				.OnValueChanged_Lambda([Curve, KeyHandle] (float NewValue) {Curve->SetKeyValue(KeyHandle, NewValue);})
+				.Justification(ETextJustify::Right)
+			;
+		}
+		return SNullWidget::NullWidget;
 	}
 
-	EVisibility GetTableViewVisibility() const
+	void CreateCurveModels(TArray<TUniquePtr<FCurveModel>>& OutCurveModels) override
 	{
-		return CurveTableEditor.Pin()->GetViewMode() == ECurveTableViewMode::Grid ? EVisibility::Visible : EVisibility::Collapsed;
+		if (RowHandle.HasRichCurves())
+		{
+			if (FRichCurve* RichCurve = RowHandle.GetRichCurve())
+			{
+				const UCurveTable* Table = RowHandle.CurveTable.Get();
+				UCurveTable* RawTable = const_cast<UCurveTable*>(Table);
+
+				TUniquePtr<FRichCurveEditorModelRaw> NewCurve = MakeUnique<FRichCurveEditorModelRaw>(RichCurve, RawTable);
+				NewCurve->SetShortDisplayName(DisplayName);
+				NewCurve->SetColor(FStyleColors::AccentOrange.GetSpecifiedColor());
+				OutCurveModels.Add(MoveTemp(NewCurve));
+			}
+		}
+		else
+		{
+			const UCurveTable* Table = RowHandle.CurveTable.Get();
+			UCurveTable* RawTable = const_cast<UCurveTable*>(Table);
+
+			TUniquePtr<FRealCurveModel> NewCurveModel = MakeUnique<FRealCurveModel>(RowHandle.GetCurve(), RawTable);
+			NewCurveModel->SetShortDisplayName(DisplayName);
+
+			OutCurveModels.Add(MoveTemp(NewCurveModel));
+		}
 	}
 
-	EVisibility GetCurveViewVisibility() const
+	bool PassesFilter(const FCurveEditorTreeFilter* InFilter) const override
 	{
-		return CurveTableEditor.Pin()->GetViewMode() == ECurveTableViewMode::CurveTable ? EVisibility::Visible : EVisibility::Collapsed;
+		if (InFilter->GetType() == ECurveEditorTreeFilterType::Text)
+		{
+			const FCurveEditorTreeTextFilter* Filter = static_cast<const FCurveEditorTreeTextFilter*>(InFilter);
+			for (const FCurveEditorTreeTextFilterTerm& Term : Filter->GetTerms())
+			{
+				for(const FCurveEditorTreeTextFilterToken& Token : Term.ChildToParentTokens)
+				{
+					if(Token.Match(*DisplayName.ToString()))
+					{
+						return true;
+					}
+				}
+			}
+
+			return false;
+		}
+
+		return false;
+	}
+
+	void CacheKeys()
+	{
+		if (!RowHandle.HasRichCurves())
+		{
+			if (FRealCurve* Curve = RowHandle.GetCurve())
+			{	
+				for (auto Col : Columns)
+				{
+					FKeyHandle KeyHandle = Curve->FindKey(Col->KeyTime);
+					float KeyValue = Curve->GetKeyValue(KeyHandle);
+
+					CellDataMap.Add(Col->ColumnId, CachedKeyInfo(KeyHandle, FText::AsNumber(KeyValue))); 
+				}
+			}
+		}
 	}
 
-private:
-	/** Weak reference to the curve table editor that owns our list */
-	TWeakPtr<FCurveTableEditor> CurveTableEditor;
-	/** The item associated with this row of data */
-	FCurveTableEditorRowListViewDataPtr Item;
+	/** Unique ID used to identify this row */
+	FName RowId;
+
+	/** Display name of this row */
+	FText DisplayName;
+
+	/** Array corresponding to each cell in this row */
+	TMap<FName, CachedKeyInfo> CellDataMap;
+
+	/** Handle to the row */
+	FCurveTableEditorHandle RowHandle;
+
+	/** A Reference to the available columns in the TableView */
+	const TArray<FCurveTableEditorColumnHeaderDataPtr>& Columns;
 };
 
 
@@ -159,16 +248,6 @@ void FCurveTableEditor::InitCurveTableEditor( const EToolkitMode::Type Mode, con
 	RegenerateMenusAndToolbars();
 
 	FReimportManager::Instance()->OnPostReimport().AddSP(this, &FCurveTableEditor::OnPostReimport);
-
-	// @todo toolkit world centric editing
-	/*// Setup our tool's layout
-	if( IsWorldCentricAssetEditor() )
-	{
-		const FString TabInitializationPayload(TEXT(""));		// NOTE: Payload not currently used for table properties
-		SpawnToolkitTab( CurveTableTabId, TabInitializationPayload, EToolkitTabSpot::Details );
-	}*/
-
-	// NOTE: Could fill in asset editor commands here!
 }
 
 TSharedRef< FTabManager::FLayout > FCurveTableEditor::InitCurveTableLayout()
@@ -272,44 +351,42 @@ TSharedRef<SDockTab> FCurveTableEditor::SpawnTab_CurveTable( const FSpawnTabArgs
 {
 	check( Args.GetTabId().TabType == CurveTableTabId );
 
-	TSharedRef<SScrollBar> HorizontalScrollBar = SNew(SScrollBar)
-		.Orientation(Orient_Horizontal)
-		.Thickness(FVector2D(12.0f, 12.0f));
+	bUpdatingTableViewSelection = false;
 
 	TSharedRef<SScrollBar> VerticalScrollBar = SNew(SScrollBar)
-		.Orientation(Orient_Vertical)
-		.Thickness(FVector2D(12.0f, 12.0f));
+		.Orientation(Orient_Vertical);
 
-	TSharedRef<SHeaderRow> RowNamesHeaderRow = SNew(SHeaderRow)
-		.Visibility(this, &FCurveTableEditor::GetGridViewControlsVisibility);
+	ColumnNamesHeaderRow = SNew(SHeaderRow)
+		.Visibility(this, &FCurveTableEditor::GetTableViewControlsVisibility);
 
-	RowNamesHeaderRow->AddColumn(
-		SHeaderRow::Column(RowNameColumnId)
-		.DefaultLabel(FText::GetEmpty())
-		);
+	CurveEditor = MakeShared<FCurveEditor>();
 
-	ColumnNamesHeaderRow = SNew(SHeaderRow)
-		.Visibility(this, &FCurveTableEditor::GetGridViewControlsVisibility);
-
-	RowNamesListView = SNew(SListView<FCurveTableEditorRowListViewDataPtr>)
-		.ListItemsSource(&AvailableRows)
-		.HeaderRow(RowNamesHeaderRow)
-		.OnGenerateRow(this, &FCurveTableEditor::MakeRowNameWidget)
-		.OnListViewScrolled(this, &FCurveTableEditor::OnRowNamesListViewScrolled)
-		.ScrollbarVisibility(EVisibility::Collapsed)
-		.ConsumeMouseWheel(EConsumeMouseWheel::Always)
-		.SelectionMode(ESelectionMode::None)
-		.AllowOverscroll(EAllowOverscroll::No);
-
-	CellsListView = SNew(SListView<FCurveTableEditorRowListViewDataPtr>)
-		.ListItemsSource(&AvailableRows)
+	FCurveEditorInitParams CurveEditorInitParams;
+	CurveEditor->InitCurveEditor(CurveEditorInitParams);
+
+	CurveEditorTree = SNew(SCurveEditorTree, CurveEditor.ToSharedRef())
+		.OnTreeViewScrolled(this, &FCurveTableEditor::OnCurveTreeViewScrolled);
+
+	TSharedRef<SCurveEditorPanel> CurveEditorPanel = SNew(SCurveEditorPanel, CurveEditor.ToSharedRef());
+
+	TableView = SNew(SListView<FCurveEditorTreeItemID>)
+		.ListItemsSource(&EmptyItems)
+		.OnListViewScrolled(this, &FCurveTableEditor::OnTableViewScrolled)
 		.HeaderRow(ColumnNamesHeaderRow)
-		.OnGenerateRow(this, &FCurveTableEditor::MakeRowWidget)
-		.OnListViewScrolled(this, &FCurveTableEditor::OnCellsListViewScrolled)
+		.OnGenerateRow(CurveEditorTree.Get(), &SCurveEditorTree::GenerateRow)
 		.ExternalScrollbar(VerticalScrollBar)
-		.ConsumeMouseWheel(EConsumeMouseWheel::Always)
-		.SelectionMode(ESelectionMode::None)
-		.AllowOverscroll(EAllowOverscroll::No);
+		.SelectionMode(ESelectionMode::Multi)
+		.OnSelectionChanged_Lambda(
+			[this](TListTypeTraits<FCurveEditorTreeItemID>::NullableType InItemID, ESelectInfo::Type Type)
+			{
+				this->OnTableViewSelectionChanged(InItemID, Type);
+			}
+		);
+
+	CurveEditor->GetTree()->Events.OnItemsChanged.AddSP(this, &FCurveTableEditor::RefreshTableRows);
+	CurveEditor->GetTree()->Events.OnSelectionChanged.AddSP(this, &FCurveTableEditor::RefreshTableRowsSelection);
+
+	ViewMode = GetCurveTable()->HasRichCurves() ? ECurveTableViewMode::CurveTable : ECurveTableViewMode::Grid;
 
 	RefreshCachedCurveTable();
 
@@ -323,95 +400,137 @@ TSharedRef<SDockTab> FCurveTableEditor::SpawnTab_CurveTable( const FSpawnTabArgs
 			[
 				SNew(SVerticalBox)
 				+SVerticalBox::Slot()
+				.AutoHeight()
+				.Padding(FMargin(8, 0))
 				[
-					SNew(SHorizontalBox)
-					+SHorizontalBox::Slot()
-					.AutoWidth()
+					MakeToolbar(CurveEditorPanel)
+				]
+
+				+SVerticalBox::Slot()
+				[
+					SNew(SSplitter)
+					+SSplitter::Slot()
+					.Value(.2)
 					[
-						SNew(SBox)
-						.WidthOverride(this, &FCurveTableEditor::GetRowNameColumnWidth)
+						SNew(SVerticalBox)
+					
+						+SVerticalBox::Slot()
+						.Padding(0, 0, 0, 1) // adjusting padding so as to line up the rows in the cell view
+						.AutoHeight()
 						[
-							RowNamesListView.ToSharedRef()
+							SNew(SHorizontalBox)
+							+SHorizontalBox::Slot()
+							.AutoWidth()
+							.Padding(2.f, 0.f, 4.f, 0.0)
+							[
+								SNew(SEditorHeaderButton)
+								.Icon(FAppStyle::Get().GetBrush("Icons.Plus"))
+								.Text(LOCTEXT("Curve", "Curve"))
+								.OnClicked(this, &FCurveTableEditor::OnAddCurveClicked)
+							]
+
+							+SHorizontalBox::Slot()	
+							[
+								SNew(SCurveEditorTreeTextFilter, CurveEditor)
+							]
 						]
-					]
-					+SHorizontalBox::Slot()
-					[
-						SNew(SScrollBox)
-						.Orientation(Orient_Horizontal)
-						.ExternalScrollbar(HorizontalScrollBar)
-						.Visibility(this, &FCurveTableEditor::GetGridViewControlsVisibility)
-						+SScrollBox::Slot()
+
+						+SVerticalBox::Slot()
 						[
-							CellsListView.ToSharedRef()
+							CurveEditorTree.ToSharedRef()
 						]
+
 					]
-					+SHorizontalBox::Slot()
+					+SSplitter::Slot()
 					[
-						SNew(SBox)
-						.Visibility(this, &FCurveTableEditor::GetCurveViewControlsVisibility)
+
+						SNew(SHorizontalBox)
+						.Visibility(this, &FCurveTableEditor::GetTableViewControlsVisibility)
+
+						+SHorizontalBox::Slot()
 						[
-							CellsListView.ToSharedRef()
+							SNew(SScrollBox)
+							.Orientation(Orient_Horizontal)
+
+							+SScrollBox::Slot()
+							[
+								TableView.ToSharedRef()
+							]
+						]
+
+						+SHorizontalBox::Slot()
+						.AutoWidth()
+						[
+							VerticalScrollBar
 						]
 					]
-					+SHorizontalBox::Slot()
-					.AutoWidth()
-					[
-						VerticalScrollBar
-					]
-				]
-				+SVerticalBox::Slot()
-				.AutoHeight()
-				[
-					SNew(SHorizontalBox)
-					+SHorizontalBox::Slot()
-					.AutoWidth()
+
+					+SSplitter::Slot()
 					[
 						SNew(SBox)
-						.WidthOverride(this, &FCurveTableEditor::GetRowNameColumnWidth)
-						.Visibility(this, &FCurveTableEditor::GetGridViewControlsVisibility)
+						.Visibility(this, &FCurveTableEditor::GetCurveViewControlsVisibility)
 						[
-							SNullWidget::NullWidget
+							CurveEditorPanel
 						]
 					]
-					+SHorizontalBox::Slot()
-					[
-						HorizontalScrollBar
-					]
 				]
 			]
 		];
 }
 
-void FCurveTableEditor::RefreshCachedCurveTable()
+void FCurveTableEditor::RefreshTableRows()
 {
-	CacheCurveTableForEditing();
+	TableView->RequestListRefresh();
+}
 
-	ColumnNamesHeaderRow->ClearColumns();
-	for (int32 ColumnIndex = 0; ColumnIndex < AvailableColumns.Num(); ++ColumnIndex)
+void FCurveTableEditor::RefreshTableRowsSelection()
+{
+	if(bUpdatingTableViewSelection == false)
 	{
-		const FCurveTableEditorColumnHeaderDataPtr& ColumnData = AvailableColumns[ColumnIndex];
+		TGuardValue<bool> SelecitonGuard(bUpdatingTableViewSelection, true);
 
-		ColumnNamesHeaderRow->AddColumn(
-			SHeaderRow::Column(ColumnData->ColumnId)
-			.DefaultLabel(ColumnData->DisplayName)
-			.FixedWidth(ColumnData->DesiredColumnWidth)
-			);
-	}
+		TArray<FCurveEditorTreeItemID> CurrentTreeWidgetSelection;
+		TableView->GetSelectedItems(CurrentTreeWidgetSelection);
+		const TMap<FCurveEditorTreeItemID, ECurveEditorTreeSelectionState>& CurrentCurveEditorTreeSelection = CurveEditor->GetTreeSelection();
+
+		TArray<FCurveEditorTreeItemID> NewTreeWidgetSelection;
+		for (const TPair<FCurveEditorTreeItemID, ECurveEditorTreeSelectionState>& CurveEditorTreeSelectionEntry : CurrentCurveEditorTreeSelection)
+		{
+			if (CurveEditorTreeSelectionEntry.Value != ECurveEditorTreeSelectionState::None)
+			{
+				NewTreeWidgetSelection.Add(CurveEditorTreeSelectionEntry.Key);
+				CurrentTreeWidgetSelection.RemoveSwap(CurveEditorTreeSelectionEntry.Key);
+			}
+		}
 
-	RowNamesListView->RequestListRefresh();
-	CellsListView->RequestListRefresh();
+		TableView->SetItemSelection(CurrentTreeWidgetSelection, false, ESelectInfo::Direct);
+		TableView->SetItemSelection(NewTreeWidgetSelection, true, ESelectInfo::Direct);
+	}
 }
 
+void FCurveTableEditor::OnTableViewSelectionChanged(FCurveEditorTreeItemID ItemID, ESelectInfo::Type)
+{
+	if (bUpdatingTableViewSelection == false)
+	{
+		TGuardValue<bool> SelecitonGuard(bUpdatingTableViewSelection, true);
+		CurveEditor->GetTree()->SetDirectSelection(TableView->GetSelectedItems(), CurveEditor.Get());
+	}
+}
 
-void FCurveTableEditor::CacheCurveTableForEditing()
+void FCurveTableEditor::RefreshCachedCurveTable()
 {
-	RowNameColumnWidth = 10.0f;
+
+	// This will trigger to remove any cached widgets in the TableView while we rebuild the model from the source CurveTable
+	TableView->SetListItemsSource(EmptyItems);
+	
+	CurveEditor->RemoveAllTreeItems();
+
+	ColumnNamesHeaderRow->ClearColumns();
+	AvailableColumns.Empty();
 
 	const UCurveTable* Table = GetCurveTable();
 	if (!Table || Table->GetRowMap().Num() == 0)
 	{
-		AvailableColumns.Empty();
-		AvailableRows.Empty();
 		return;
 	}
 
@@ -419,167 +538,74 @@ void FCurveTableEditor::CacheCurveTableForEditing()
 	const FTextBlockStyle& CellTextStyle = FEditorStyle::GetWidgetStyle<FTextBlockStyle>("DataTableEditor.CellText");
 	static const float CellPadding = 10.0f;
 
-	// Find unique column titles
-	TArray<float> UniqueColumns;
-	for (const TPair<FName, FRealCurve*>& CurveRow : Table->GetRowMap())
+	if (Table->HasRichCurves())
 	{
-		FRealCurve* Curve = CurveRow.Value;
-		for (auto CurveIt(Curve->GetKeyHandleIterator()); CurveIt; ++CurveIt)
+		for (const TPair<FName, FRichCurve*>& CurveRow : Table->GetRichCurveRowMap())
 		{
-			UniqueColumns.AddUnique(Curve->GetKeyTime(*CurveIt));
+			// Setup the CurveEdtiorTree
+			const FName& CurveName = CurveRow.Key;
+			FCurveEditorTreeItem* TreeItem = CurveEditor->AddTreeItem(FCurveEditorTreeItemID());
+			TreeItem->SetStrongItem(MakeShared<FCurveTableEditorItem>(CurveName, FCurveTableEditorHandle(Table, CurveName), AvailableColumns));
 		}
 	}
 
-	UniqueColumns.Sort();
-
-	AvailableColumns.Empty();
-	for (const float& ColumnTime : UniqueColumns)
+	else
 	{
-		const FText ColumnText = FText::AsNumber(ColumnTime);
-
-		FCurveTableEditorColumnHeaderDataPtr CachedColumnData = MakeShareable(new FCurveTableEditorColumnHeaderData());
-		CachedColumnData->ColumnId = *ColumnText.ToString();
-		CachedColumnData->DisplayName = ColumnText;
-		CachedColumnData->DesiredColumnWidth = FontMeasure->Measure(CachedColumnData->DisplayName, CellTextStyle.Font).X + CellPadding;
-
-		AvailableColumns.Add(CachedColumnData);
-	}
-
-	// Each curve is a row entry
-	AvailableRows.Reset(Table->GetRowMap().Num());
-	for (const TPair<FName, FRealCurve*>& CurveRow : Table->GetRowMap())
-	{
-		const FName& CurveName = CurveRow.Key;
-		FRealCurve* Curve = CurveRow.Value;
-
-		FCurveTableEditorRowListViewDataPtr CachedRowData = MakeShareable(new FCurveTableEditorRowListViewData());
-		CachedRowData->RowId = CurveName;
-		CachedRowData->DisplayName = FText::FromName(CurveName);
-		CachedRowData->RowHandle = FCurveTableEditorHandle(Table, CurveName);
-
-		check(CachedRowData->RowHandle.IsValid());
-
-		const float RowNameWidth = FontMeasure->Measure(CachedRowData->DisplayName, CellTextStyle.Font).X + CellPadding;
-		RowNameColumnWidth = FMath::Max(RowNameColumnWidth, RowNameWidth);
-
-		CachedRowData->CellData.AddDefaulted(AvailableColumns.Num());
-
-		for (auto It(Curve->GetKeyHandleIterator()); It; ++It)
+		// Find unique column titles and setup columns
+		TArray<float> UniqueColumns;
+		for (const TPair<FName, FRealCurve*>& CurveRow : Table->GetRowMap())
 		{
-			const FKeyHandle& KeyHandle = *It;
-			int32 ColumnIndex = 0;
-			const TPair<float, float> TimeValuePair = Curve->GetKeyTimeValuePair(KeyHandle);
-			if (UniqueColumns.Find(TimeValuePair.Key, ColumnIndex))
+			FRealCurve* Curve = CurveRow.Value;
+			for (auto CurveIt(Curve->GetKeyHandleIterator()); CurveIt; ++CurveIt)
 			{
-				FCurveTableEditorColumnHeaderDataPtr CachedColumnData = AvailableColumns[ColumnIndex];
-
-				const FText CellText = FText::AsNumber(TimeValuePair.Value);
-				CachedRowData->CellData[ColumnIndex] = CellText;
-
-				const float CellWidth = FontMeasure->Measure(CellText, CellTextStyle.Font).X + CellPadding;
-				CachedColumnData->DesiredColumnWidth = FMath::Max(CachedColumnData->DesiredColumnWidth, CellWidth);
+				UniqueColumns.AddUnique(Curve->GetKeyTime(*CurveIt));
 			}
 		}
-
-		AvailableRows.Add(CachedRowData);
-	}
-}
-
-
-TSharedRef<ITableRow> FCurveTableEditor::MakeRowNameWidget(FCurveTableEditorRowListViewDataPtr InRowDataPtr, const TSharedRef<STableViewBase>& OwnerTable)
-{
-	TWeakPtr<FCurveTableEditor> WeakEditorPtr = SharedThis(this);
-
-	return
-		SNew(STableRow<FCurveTableEditorRowListViewDataPtr>, OwnerTable)
-		.Style(FEditorStyle::Get(), "DataTableEditor.NameListViewRow")
-		[
-			SNew(SBox)
-			.Padding(FMargin(4, 2, 4, 2))
-			.HeightOverride_Lambda([WeakEditorPtr]() { TSharedPtr<FCurveTableEditor> This = WeakEditorPtr.Pin(); return (This.IsValid() && This->GetViewMode() == ECurveTableViewMode::CurveTable) ? FOptionalSize(68.0f) : FOptionalSize(); })
-			.VAlign(VAlign_Center)
-			[
-				SNew(STextBlock)
-				.Text(InRowDataPtr->DisplayName)
-			]
-		];
-}
-
-
-TSharedRef<ITableRow> FCurveTableEditor::MakeRowWidget(FCurveTableEditorRowListViewDataPtr InRowDataPtr, const TSharedRef<STableViewBase>& OwnerTable)
-{
-	return
-		SNew(SCurveTableListViewRow, OwnerTable)
-		.CurveTableEditor(SharedThis(this))
-		.Item(InRowDataPtr);
-}
-
-
-TSharedRef<SWidget> FCurveTableEditor::MakeCellWidget(FCurveTableEditorRowListViewDataPtr InRowDataPtr, const int32 InRowIndex, const FName& InColumnId)
-{
-	int32 ColumnIndex = 0;
-	for (; ColumnIndex < AvailableColumns.Num(); ++ColumnIndex)
-	{
-		const FCurveTableEditorColumnHeaderDataPtr& ColumnData = AvailableColumns[ColumnIndex];
-		if (ColumnData->ColumnId == InColumnId)
+		UniqueColumns.Sort();
+		for (const float& ColumnTime : UniqueColumns)
 		{
-			break;
+			const FText ColumnText = FText::AsNumber(ColumnTime);
+			FCurveTableEditorColumnHeaderDataPtr CachedColumnData = MakeShareable(new FCurveTableEditorColumnHeaderData());
+			CachedColumnData->ColumnId = *ColumnText.ToString();
+			CachedColumnData->DisplayName = ColumnText;
+			CachedColumnData->DesiredColumnWidth = FontMeasure->Measure(CachedColumnData->DisplayName, CellTextStyle.Font).X + CellPadding;
+			CachedColumnData->KeyTime = ColumnTime;
+
+			AvailableColumns.Add(CachedColumnData);
+
+			ColumnNamesHeaderRow->AddColumn(
+				SHeaderRow::Column(CachedColumnData->ColumnId)
+				.DefaultLabel(CachedColumnData->DisplayName)
+				.FixedWidth(CachedColumnData->DesiredColumnWidth + 50)
+				.HAlignHeader(HAlign_Center)
+			);
 		}
-	}
 
-	// Valid column ID?
-	if (AvailableColumns.IsValidIndex(ColumnIndex))
-	{
-		FText ColumnText;
-		if (InRowDataPtr->CellData.IsValidIndex(ColumnIndex))
+		// Setup the CurveEditorTree 
+		for (const TPair<FName, FRealCurve*>& CurveRow : Table->GetRowMap())
 		{
-			ColumnText = InRowDataPtr->CellData[ColumnIndex];
+			const FName& CurveName = CurveRow.Key;
+			FCurveEditorTreeItem* TreeItem = CurveEditor->AddTreeItem(FCurveEditorTreeItemID());
+			TSharedPtr<FCurveTableEditorItem> NewItem = MakeShared<FCurveTableEditorItem>(CurveName, FCurveTableEditorHandle(Table, CurveName), AvailableColumns);
+			OnColumnsChanged.AddSP(NewItem.ToSharedRef(), &FCurveTableEditorItem::CacheKeys);
+			TreeItem->SetStrongItem(NewItem);
 		}
-
-		return SNew(SBox)
-			.Padding(FMargin(4, 2, 4, 2))
-			[
-				SNew(STextBlock)
-				.TextStyle(FEditorStyle::Get(), "DataTableEditor.CellText")
-				.Text(ColumnText)
-			];
 	}
 
-	return SNullWidget::NullWidget;
+	TableView->SetListItemsSource(CurveEditorTree->GetSourceItems());
 }
 
-TSharedRef<SWidget> FCurveTableEditor::MakeCurveWidget(FCurveTableEditorRowListViewDataPtr InRowDataPtr, const int32 InRowIndex)
-{
-	TSharedRef<SCurveTableEditor> CurveEditor = SNew(SCurveTableEditor);
-	CurveEditor->SetCurveOwner(&InRowDataPtr->RowHandle, false);
-
-	return SNew(SHorizontalBox)
-		+SHorizontalBox::Slot()
-		.Padding(FMargin(4, 2, 4, 2))
-		.FillWidth(1.0f)
-		[
-			CurveEditor
-		];
-}
-
-
-void FCurveTableEditor::OnRowNamesListViewScrolled(double InScrollOffset)
+void FCurveTableEditor::OnCurveTreeViewScrolled(double InScrollOffset)
 {
 	// Synchronize the list views
-	CellsListView->SetScrollOffset(InScrollOffset);
+	TableView->SetScrollOffset(InScrollOffset);
 }
 
 
-void FCurveTableEditor::OnCellsListViewScrolled(double InScrollOffset)
+void FCurveTableEditor::OnTableViewScrolled(double InScrollOffset)
 {
 	// Synchronize the list views
-	RowNamesListView->SetScrollOffset(InScrollOffset);
-}
-
-
-FOptionalSize FCurveTableEditor::GetRowNameColumnWidth() const
-{
-	return FOptionalSize(RowNameColumnWidth);
+	CurveEditorTree->SetScrollOffset(InScrollOffset);
 }
 
 void FCurveTableEditor::OnPostReimport(UObject* InObject, bool)
@@ -591,7 +617,7 @@ void FCurveTableEditor::OnPostReimport(UObject* InObject, bool)
 	}
 }
 
-EVisibility FCurveTableEditor::GetGridViewControlsVisibility() const
+EVisibility FCurveTableEditor::GetTableViewControlsVisibility() const
 {
 	return ViewMode == ECurveTableViewMode::CurveTable ? EVisibility::Collapsed : EVisibility::Visible;
 }
@@ -611,4 +637,164 @@ bool FCurveTableEditor::IsCurveViewChecked() const
 	return (ViewMode == ECurveTableViewMode::CurveTable);
 }
 
+TSharedRef<SWidget> FCurveTableEditor::MakeToolbar(TSharedRef<SCurveEditorPanel>& InEditorPanel)
+{
+	FToolBarBuilder ToolBarBuilder(InEditorPanel->GetCommands(), FMultiBoxCustomization::None, InEditorPanel->GetToolbarExtender(), true);
+	ToolBarBuilder.SetStyle(&FAppStyle::Get(), "Sequencer.ToolBar");
+	ToolBarBuilder.BeginSection("Asset");
+	ToolBarBuilder.EndSection();
+	// We just use all of the extenders as our toolbar, we don't have a need to create a separate toolbar.
+
+	bool HasRichCurves = GetCurveTable()->HasRichCurves();
+
+	return SNew(SHorizontalBox)
+
+	+SHorizontalBox::Slot()
+	.AutoWidth()
+	.VAlign(VAlign_Center)
+	.HAlign(HAlign_Center)
+	.Padding(FMargin(2.0, 4.0, 8.f, 4.f))
+	[
+		SNew(SSegmentedControl<ECurveTableViewMode>)
+		.Visibility(HasRichCurves ? EVisibility::Collapsed : EVisibility::Visible)
+		.OnValueChanged_Lambda([this] (ECurveTableViewMode InMode) {if (InMode != GetViewMode()) ToggleViewMode();  } )
+		.Value(this, &FCurveTableEditor::GetViewMode)
+
+		+SSegmentedControl<ECurveTableViewMode>::Slot(ECurveTableViewMode::CurveTable)
+        .Icon(FAppStyle::Get().GetBrush("CurveTableEditor.CurveView"))
+
+		+SSegmentedControl<ECurveTableViewMode>::Slot(ECurveTableViewMode::Grid)
+        .Icon(FAppStyle::Get().GetBrush("CurveTableEditor.TableView"))
+	]
+
+	+SHorizontalBox::Slot()
+	.AutoWidth()
+	.VAlign(VAlign_Center)
+	[
+		SNew(SBox)
+		[
+			SNew(SButton)
+            .ButtonStyle( &FAppStyle::Get().GetWidgetStyle< FButtonStyle >( "SimpleButton" ) )
+			.Visibility(this, &FCurveTableEditor::GetTableViewControlsVisibility)
+			.OnClicked(this, &FCurveTableEditor::OnAddNewKeyColumn)
+			.ToolTipText(LOCTEXT("CurveTableEditor_AddKeyColumnTooltip", "Append a new column to the curve table.\nEvery Curve or Table Row will have a new key appended."))
+			[
+                SNew(SImage)
+                .ColorAndOpacity(FSlateColor::UseForeground())
+                .Image(FAppStyle::Get().GetBrush("Sequencer.KeyTriangle")) 
+	        ]
+		]
+	]
+
+	+SHorizontalBox::Slot()
+	.AutoWidth()
+	.VAlign(VAlign_Center)
+	[
+		SNew(SBox)
+		.Visibility(this, &FCurveTableEditor::GetCurveViewControlsVisibility)
+		[
+			ToolBarBuilder.MakeWidget()
+		]
+	];
+}
+
+FReply FCurveTableEditor::OnAddCurveClicked()
+{
+	UCurveTable* Table = Cast<UCurveTable>(GetEditingObject());
+	check(Table != nullptr);
+
+	if (Table->HasRichCurves())
+	{
+		FName NewCurveUnique = MakeUniqueObjectName(Table, UCurveTable::StaticClass(), "Curve");
+		FRichCurve& NewCurve = Table->AddRichCurve(NewCurveUnique);
+		FCurveEditorTreeItem* TreeItem = CurveEditor->AddTreeItem(FCurveEditorTreeItemID());
+		TreeItem->SetStrongItem(MakeShared<FCurveTableEditorItem>(NewCurveUnique, FCurveTableEditorHandle(Table, NewCurveUnique), AvailableColumns));
+	}
+	else
+	{
+		FName NewCurveUnique = MakeUniqueObjectName(Table, UCurveTable::StaticClass(), "Curve");
+		FRealCurve& RealCurve = Table->AddSimpleCurve(NewCurveUnique);
+		// Also add a default key for each column 
+		for (auto Column : AvailableColumns)
+		{
+			RealCurve.AddKey(Column->KeyTime, 0.0);
+		}
+
+		FCurveEditorTreeItem* TreeItem = CurveEditor->AddTreeItem(FCurveEditorTreeItemID());
+		TSharedPtr<FCurveTableEditorItem> NewItem = MakeShared<FCurveTableEditorItem>(NewCurveUnique, FCurveTableEditorHandle(Table, NewCurveUnique), AvailableColumns);
+		OnColumnsChanged.AddSP(NewItem.ToSharedRef(), &FCurveTableEditorItem::CacheKeys);
+		TreeItem->SetStrongItem(NewItem);
+
+
+	}
+
+	return FReply::Handled();
+}
+
+FReply FCurveTableEditor::OnAddNewKeyColumn()
+{
+	UCurveTable* Table = Cast<UCurveTable>(GetEditingObject());
+	check(Table != nullptr);
+
+	if (!Table->HasRichCurves())
+	{
+		// Compute a new keytime based on the last columns 
+		float NewKeyTime = 1.0;
+		if (AvailableColumns.Num() > 1)
+		{
+			float LastKeyTime = AvailableColumns[AvailableColumns.Num() - 1]->KeyTime;
+			float PrevKeyTime = AvailableColumns[AvailableColumns.Num() - 2]->KeyTime;
+			NewKeyTime = 2.*LastKeyTime - PrevKeyTime;
+		}
+		else if (AvailableColumns.Num() > 0)
+		{
+			float LastKeyTime = AvailableColumns[AvailableColumns.Num() - 1]->KeyTime;
+			NewKeyTime = LastKeyTime + 1;
+		}
+
+		AddNewKeyColumn(NewKeyTime);
+	}
+
+	return FReply::Handled();
+}
+
+void FCurveTableEditor::AddNewKeyColumn(float NewKeyTime)
+{
+	UCurveTable* Table = Cast<UCurveTable>(GetEditingObject());
+	check(Table != nullptr);
+
+	if (!Table->HasRichCurves())
+	{
+		// Make sure we don't already have a key at this time
+
+		// 1. Add new keys to every curve
+		for (const TPair<FName, FRealCurve*>& CurveRow : Table->GetRowMap())
+		{
+			FRealCurve* Curve = CurveRow.Value;
+			Curve->UpdateOrAddKey(NewKeyTime, Curve->Eval(NewKeyTime));
+		}
+
+		// 2. Add Column to our Table
+		FCurveTableEditorColumnHeaderDataPtr ColumnData = MakeShareable(new FCurveTableEditorColumnHeaderData());
+		const FText ColumnText = FText::AsNumber(NewKeyTime);
+		ColumnData->ColumnId = *ColumnText.ToString();
+		ColumnData->DisplayName = ColumnText;
+		// ColumnData->DesiredColumnWidth = FontMeasure->Measure(ColumnData->DisplayName, CellTextStyle.Font).X + CellPadding;
+		ColumnData->KeyTime = NewKeyTime;
+
+		AvailableColumns.Add(ColumnData);
+
+		// 3. Let the CurveTreeItems know they need to recache
+		OnColumnsChanged.Broadcast();
+
+		// Add the column to the TableView Header Row
+		ColumnNamesHeaderRow->AddColumn(
+			SHeaderRow::Column(ColumnData->ColumnId)
+			.DefaultLabel(ColumnData->DisplayName)
+			.FixedWidth(ColumnData->DesiredColumnWidth + 50)
+			.HAlignHeader(HAlign_Center)
+		);	
+	}
+}
+
 #undef LOCTEXT_NAMESPACE
diff --git a/Engine/Source/Editor/CurveTableEditor/Private/CurveTableEditor.h b/Engine/Source/Editor/CurveTableEditor/Private/CurveTableEditor.h
--- a/Engine/Source/Editor/CurveTableEditor/Private/CurveTableEditor.h
+++ b/Engine/Source/Editor/CurveTableEditor/Private/CurveTableEditor.h
@@ -10,35 +10,15 @@
 #include "ICurveTableEditor.h"
 #include "Widgets/Views/STableViewBase.h"
 #include "Widgets/Views/STableRow.h"
+#include "Tree/ICurveEditorTreeItem.h"
 #include "CurveTableEditorHandle.h"
 #include "CurveTableEditorUtils.h"
 
-struct FCurveTableEditorColumnHeaderData
-{
-	/** Unique ID used to identify this column */
-	FName ColumnId;
-
-	/** Display name of this column */
-	FText DisplayName;
-
-	/** The calculated width of this column taking into account the cell data for each row */
-	float DesiredColumnWidth;
-};
-
-struct FCurveTableEditorRowListViewData
-{
-	/** Unique ID used to identify this row */
-	FName RowId;
-
-	/** Display name of this row */
-	FText DisplayName;
+#include "CurveEditorTypes.h"
+class FCurveEditor;
+class SCurveEditorTree;
+class SCurveEditorPanel;
 
-	/** Array corresponding to each cell in this row */
-	TArray<FText> CellData;
-
-	/** Handle to the row */
-	FCurveTableEditorHandle RowHandle;
-};
 
 /** The manner in which curve tables are displayed */
 enum class ECurveTableViewMode : int32
@@ -50,16 +30,14 @@ enum class ECurveTableViewMode : int32
 	CurveTable,
 };
 
+struct FCurveTableEditorColumnHeaderData;
 typedef TSharedPtr<FCurveTableEditorColumnHeaderData> FCurveTableEditorColumnHeaderDataPtr;
-typedef TSharedPtr<FCurveTableEditorRowListViewData>  FCurveTableEditorRowListViewDataPtr;
 
-/** Viewer/editor for a CurveTable */
+/** Viewer/Editor for a CurveTable */
 class FCurveTableEditor :
 	public ICurveTableEditor
 	, public FCurveTableEditorUtils::INotifyOnCurveTableChanged
 {
-	friend class SCurveTableListViewRow;
-	friend class SCurveTableCurveViewRow;
 
 public:
 	virtual void RegisterTabSpawners(const TSharedRef<class FTabManager>& TabManager) override;
@@ -112,39 +90,36 @@ protected:
 	/** Update the cached state of this curve table, and then reflect that new state in the UI */
 	void RefreshCachedCurveTable();
 
-	/** Cache the data from the current curve table so that it can be shown in the editor */
-	void CacheCurveTableForEditing();
-
-	/** Make the widget for a row name entry in the data table row list view */
-	TSharedRef<ITableRow> MakeRowNameWidget(FCurveTableEditorRowListViewDataPtr InRowDataPtr, const TSharedRef<STableViewBase>& OwnerTable);
-
-	/** Make the widget for a row entry in the data table row list view */
-	TSharedRef<ITableRow> MakeRowWidget(FCurveTableEditorRowListViewDataPtr InRowDataPtr, const TSharedRef<STableViewBase>& OwnerTable);
+	/** Make the toolbar */
+	TSharedRef<SWidget> MakeToolbar(TSharedRef<SCurveEditorPanel>& CurveEditorPanel);
 
-	/** Make the widget for a cell entry in the data table row list view */
-	TSharedRef<SWidget> MakeCellWidget(FCurveTableEditorRowListViewDataPtr InRowDataPtr, const int32 InRowIndex, const FName& InColumnId);
+	/** Called when the CurveEditorTree view is scrolled - used to keep the two list views in sync */
+	void OnCurveTreeViewScrolled(double InScrollOffset);
 
-	/** Make the curve widget for a row entry in the data table row list view */
-	TSharedRef<SWidget> MakeCurveWidget(FCurveTableEditorRowListViewDataPtr InRowDataPtr, const int32 InRowIndex);
+	/** Called when the Table View is scrolled - used to keep the two list views in sync */
+	void OnTableViewScrolled(double InScrollOffset);
 
-	/** Called when the row names list is scrolled - used to keep the two list views in sync */
-	void OnRowNamesListViewScrolled(double InScrollOffset);
-
-	/** Called when the cell names list view is scrolled - used to keep the two list views in sync */
-	void OnCellsListViewScrolled(double InScrollOffset);
-
-	/** Get the width to use for the row names column */
-	FOptionalSize GetRowNameColumnWidth() const;
+	/** Called when someone selected a row directly in the TableView - used to keep selection in sync between CurveTree  and TableView */
+	void OnTableViewSelectionChanged(FCurveEditorTreeItemID ItemID, ESelectInfo::Type);
 
 	/** Called when an asset has finished being imported */
 	void OnPostReimport(UObject* InObject, bool);
 
 	/** Control control visibility based on view mode */
-	EVisibility GetGridViewControlsVisibility() const;
+	EVisibility GetTableViewControlsVisibility() const;
 
 	/** Control control visibility based on view mode */
 	EVisibility GetCurveViewControlsVisibility() const;
 
+	/** Add New Curve Callback */
+	FReply OnAddCurveClicked();
+
+	/** Callback For SimpleCurves, add a new Key/Column */
+	FReply OnAddNewKeyColumn();
+
+	/* Adds new key for all (Simple) curves in the table at given time */
+	void AddNewKeyColumn(float NewKeyTime);
+
 	/** Toggle between curve & grid view */
 	void ToggleViewMode();
 
@@ -157,30 +132,39 @@ protected:
 	/** Array of the columns that are available for editing */
 	TArray<FCurveTableEditorColumnHeaderDataPtr> AvailableColumns;
 
-	/** Array of the rows that are available for editing */
-	TArray<FCurveTableEditorRowListViewDataPtr> AvailableRows;
-
 	/** Header row containing entries for each column in AvailableColumns */
 	TSharedPtr<SHeaderRow> ColumnNamesHeaderRow;
 
-	/** List view responsible for showing the row names column */
-	TSharedPtr<SListView<FCurveTableEditorRowListViewDataPtr>> RowNamesListView;
-
 	/** List view responsible for showing the rows from AvailableColumns */
-	TSharedPtr<SListView<FCurveTableEditorRowListViewDataPtr>> CellsListView;
+	TSharedPtr<SListView<FCurveEditorTreeItemID>> TableView;
 
 	/** Menu extender */
 	TSharedPtr<FExtender> MenuExtender;
 
-	/** Width of the row name column */
-	float RowNameColumnWidth;
-
 	/**	The tab id for the curve table tab */
 	static const FName CurveTableTabId;
 
-	/** The column id for the row name list view column */
-	static const FName RowNameColumnId;
-
 	/** The manner in which curve tables are displayed */
 	ECurveTableViewMode ViewMode;
+
+	/** The Curve Editor */
+	TSharedPtr<FCurveEditor> CurveEditor;
+
+	/* The Data Model that holds the source items for Views (TreeView, TableView) */
+	TSharedPtr<class SCurveEditorTree> CurveEditorTree;
+
+	bool bUpdatingTableViewSelection;
+
+	/* Sync Filtered from the CurveEditorTree model to the TableView*/
+	void RefreshTableRows();
+
+	/* Sync selected rows from the CurveEditorTree model to the TableView */
+	void RefreshTableRowsSelection();
+
+	/** A delegate to let item rows know when the number of columns have changed */
+	FSimpleMulticastDelegate OnColumnsChanged;
+
+	/** An empty source list used to initialize or when rebuilding the TableView */
+	TArray<FCurveEditorTreeItemID> EmptyItems;
+
 };
diff --git a/Engine/Source/Editor/CurveTableEditor/Private/CurveTableEditorHandle.cpp b/Engine/Source/Editor/CurveTableEditor/Private/CurveTableEditorHandle.cpp
--- a/Engine/Source/Editor/CurveTableEditor/Private/CurveTableEditorHandle.cpp
+++ b/Engine/Source/Editor/CurveTableEditor/Private/CurveTableEditorHandle.cpp
@@ -11,6 +11,24 @@ FRealCurve* FCurveTableEditorHandle::GetCurve() const
 	return nullptr;
 }
 
+FRichCurve* FCurveTableEditorHandle::GetRichCurve() const
+{
+	if (CurveTable != nullptr && RowName != NAME_None)
+	{
+		return CurveTable.Get()->FindRichCurve(RowName, TEXT("CurveTableEditorHandle::GetCurve"));
+	}
+	return nullptr;
+}
+
+bool FCurveTableEditorHandle::HasRichCurves() const
+{
+	if (CurveTable != nullptr)
+	{
+		return CurveTable.Get()->HasRichCurves();
+	}
+	return false;
+}
+
 TArray<FRichCurveEditInfoConst> FCurveTableEditorHandle::GetCurves() const
 {
 	TArray<FRichCurveEditInfoConst> Curves;
diff --git a/Engine/Source/Editor/CurveTableEditor/Private/CurveTableEditorHandle.h b/Engine/Source/Editor/CurveTableEditor/Private/CurveTableEditorHandle.h
--- a/Engine/Source/Editor/CurveTableEditor/Private/CurveTableEditorHandle.h
+++ b/Engine/Source/Editor/CurveTableEditor/Private/CurveTableEditorHandle.h
@@ -57,4 +57,8 @@ struct FCurveTableEditorHandle : public FCurveOwnerInterface
 
 	/** Get the curve straight from the row handle */
 	FRealCurve* GetCurve() const;
+	FRichCurve* GetRichCurve() const;
+
+	/** Returns true if the owning table uses RichCurves instead of Real Curves */
+	bool HasRichCurves() const;
 };
diff --git a/Engine/Source/Editor/CurveTableEditor/Private/RealCurveModel.cpp b/Engine/Source/Editor/CurveTableEditor/Private/RealCurveModel.cpp
new file mode 100644
--- /dev/null
+++ b/Engine/Source/Editor/CurveTableEditor/Private/RealCurveModel.cpp
@@ -0,0 +1,209 @@
+// Copyright Epic Games, Inc. All Rights Reserved.
+
+#include "RealCurveModel.h"
+#include "CurveDrawInfo.h"
+
+#include "Curves/RealCurve.h"
+
+
+FRealCurveModel::FRealCurveModel(FRealCurve* InRealCurve, UObject* InOwner)
+	: WeakOwner(InOwner)
+{
+
+	RealCurve = InRealCurve;
+}
+
+const void* FRealCurveModel::GetCurve() const
+{
+	return RealCurve;
+}
+
+void FRealCurveModel::Modify()
+{
+	if (UObject* Owner = WeakOwner.Get())
+	{
+		if(IsValid())
+		{
+			Owner->SetFlags(RF_Transactional);
+			Owner->Modify();
+		}
+	}	
+}
+
+void FRealCurveModel::DrawCurve(const FCurveEditor& CurveEditor, const FCurveEditorScreenSpace& ScreenSpace, TArray<TTuple<double, double>>& InterpolatingPoints) const
+{
+
+	if (IsValid())
+	{
+		for (auto It = RealCurve->GetKeyHandleIterator(); It; ++It)
+		{
+			auto KeyPair = RealCurve->GetKeyTimeValuePair(*It);
+			InterpolatingPoints.Add(MakeTuple( double(KeyPair.Key), double(KeyPair.Value)));
+		}
+	}
+}
+
+void FRealCurveModel::GetKeys(const FCurveEditor& CurveEditor, double MinTime, double MaxTime, double MinValue, double MaxValue, TArray<FKeyHandle>& OutKeyHandles) const
+{
+	if (IsValid())
+	{
+		for (auto It = RealCurve->GetKeyHandleIterator(); It; ++It)
+		{
+			auto KeyPair = RealCurve->GetKeyTimeValuePair(*It);
+			if (KeyPair.Key >= MinTime && KeyPair.Key <= MaxTime && KeyPair.Value >= MinValue && KeyPair.Value <= MaxValue)
+			{
+				OutKeyHandles.Add(*It);
+			}
+		}
+	}
+}
+
+void FRealCurveModel::AddKeys(TArrayView<const FKeyPosition> InKeyPositions, TArrayView<const FKeyAttributes> InAttributes, TArrayView<TOptional<FKeyHandle>>* OutKeyHandles)
+{
+	// Left empty for now - The curve editor cannot directly add/remove keys to RealCurves
+}
+
+void FRealCurveModel::RemoveKeys(TArrayView<const FKeyHandle> InKeys)
+{
+	// Left empty for now - The curve editor cannot directly add/remove keys to RealCurves
+}
+
+void FRealCurveModel::GetKeyPositions(TArrayView<const FKeyHandle> InKeys, TArrayView<FKeyPosition> OutKeyPositions) const
+{
+	if (IsValid())
+	{
+		for (int32 Index = 0; Index < InKeys.Num(); ++Index)
+		{
+			if (RealCurve->IsKeyHandleValid(InKeys[Index]))
+			{
+				auto KeyPair = RealCurve->GetKeyTimeValuePair(InKeys[Index]);
+				OutKeyPositions[Index].InputValue = KeyPair.Key;
+				OutKeyPositions[Index].OutputValue = KeyPair.Value;
+			}
+		}
+	}
+}
+
+void FRealCurveModel::SetKeyPositions(TArrayView<const FKeyHandle> InKeys, TArrayView<const FKeyPosition> InKeyPositions, EPropertyChangeType::Type ChangeType)
+{
+	if (UObject* Owner = WeakOwner.Get())
+	{
+		if (IsValid())
+		{
+			Owner->Modify();
+
+			for (int32 Index = 0; Index < InKeys.Num(); ++Index)
+			{
+				if (RealCurve->IsKeyHandleValid(InKeys[Index]))
+				{
+					// For now, you can only change RealCurves' position, not time 
+					RealCurve->SetKeyValue(InKeys[Index], InKeyPositions[Index].OutputValue);
+				}	
+			}
+			FPropertyChangedEvent PropertyChangeStruct(nullptr, EPropertyChangeType::ValueSet);
+			Owner->PostEditChangeProperty(PropertyChangeStruct);
+
+			CurveModifiedDelegate.Broadcast();
+		}	
+	}
+}
+
+void FRealCurveModel::GetKeyDrawInfo(ECurvePointType PointType, const FKeyHandle InKeyHandle, FKeyDrawInfo& OutDrawInfo) const
+{
+	if (PointType == ECurvePointType::ArriveTangent || PointType == ECurvePointType::LeaveTangent)
+	{
+		OutDrawInfo.Brush = FAppStyle::Get().GetBrush("GenericCurveEditor.TangentHandle");
+		OutDrawInfo.ScreenSize = FVector2D(9, 9);
+	}
+	else
+	{
+		// All keys are the same size by default
+		OutDrawInfo.ScreenSize = FVector2D(11, 11);
+
+		ERichCurveInterpMode KeyType = (IsValid() && RealCurve->IsKeyHandleValid(InKeyHandle)) ? RealCurve->GetKeyInterpMode(InKeyHandle) : RCIM_None;
+		switch (KeyType)
+		{
+		case ERichCurveInterpMode::RCIM_Constant:
+			OutDrawInfo.Brush = FAppStyle::Get().GetBrush("GenericCurveEditor.ConstantKey");
+			OutDrawInfo.Tint = FLinearColor(0, 0.45f, 0.70f);
+			break;
+		case ERichCurveInterpMode::RCIM_Linear:
+			OutDrawInfo.Brush = FAppStyle::Get().GetBrush("GenericCurveEditor.LinearKey");
+			OutDrawInfo.Tint = FLinearColor(0, 0.62f, 0.46f);
+			break;
+		case ERichCurveInterpMode::RCIM_Cubic:
+			OutDrawInfo.Brush = FAppStyle::Get().GetBrush("GenericCurveEditor.CubicKey");
+			OutDrawInfo.Tint = FLinearColor::White;
+			break;
+		default:
+			OutDrawInfo.Brush = FAppStyle::Get().GetBrush("GenericCurveEditor.Key");
+			OutDrawInfo.Tint = FLinearColor::White;
+			break;
+		}
+	}
+}
+
+void FRealCurveModel::GetKeyAttributes(TArrayView<const FKeyHandle> InKeys, TArrayView<FKeyAttributes> OutAttributes) const
+{
+	if(IsValid())
+	{
+		for (int32 Index = 0; Index < InKeys.Num(); ++Index)
+		{
+			if (RealCurve->IsKeyHandleValid(InKeys[Index]))
+			{
+				FKeyAttributes&      Attributes = OutAttributes[Index];
+				Attributes.SetInterpMode( RealCurve->GetKeyInterpMode(InKeys[Index]) );
+			}
+		}
+	}
+}
+
+void FRealCurveModel::GetTimeRange(double& MinTime, double& MaxTime) const
+{
+	if (IsValid())
+	{
+		float MinT, MaxT; 
+		RealCurve->GetTimeRange(MinT, MaxT);
+		MinTime = MinT;
+		MaxTime = MaxT;
+	}
+}
+
+void FRealCurveModel::GetValueRange(double& MinValue, double& MaxValue) const
+{
+	if (IsValid())
+	{
+		float MinV, MaxV;
+		RealCurve->GetValueRange(MinV, MaxV);
+		MinValue = MinV;
+		MaxValue = MaxV;
+	}
+}
+
+int32 FRealCurveModel::GetNumKeys() const
+{
+	if (IsValid())
+	{
+		return RealCurve->GetNumKeys();
+	}
+	return 0;
+}
+
+void FRealCurveModel::GetNeighboringKeys(const FKeyHandle InKeyHandle, TOptional<FKeyHandle>& OutPreviousKeyHandle, TOptional<FKeyHandle>& OutNextKeyHandle) const
+{
+	if (IsValid())
+	{
+		OutPreviousKeyHandle = RealCurve->GetPreviousKey(InKeyHandle);
+		OutNextKeyHandle = RealCurve->GetNextKey(InKeyHandle);
+	}	
+}
+
+bool FRealCurveModel::Evaluate(double ProspectiveTime, double& OutValue) const
+{
+	if (IsValid())
+	{
+		OutValue = RealCurve->Eval(ProspectiveTime);
+		return true;
+	}
+	return false;
+}
diff --git a/Engine/Source/Editor/CurveTableEditor/Private/RealCurveModel.h b/Engine/Source/Editor/CurveTableEditor/Private/RealCurveModel.h
new file mode 100644
--- /dev/null
+++ b/Engine/Source/Editor/CurveTableEditor/Private/RealCurveModel.h
@@ -0,0 +1,51 @@
+// Copyright Epic Games, Inc. All Rights Reserved.
+#pragma once
+
+#include "CoreTypes.h"
+#include "CurveModel.h"
+
+struct FRealCurve;
+
+/*
+ * RealCurveModel implements the FCurveModel interface for FRealCurves which allows
+ * users to edit FRealCurves in the CurveEditor
+ */
+
+class FRealCurveModel : public FCurveModel
+{
+  public:
+	static ECurveEditorViewID EventView;
+
+	FRealCurveModel(FRealCurve* InRealCurve, UObject* InOwner);
+
+	// FCurveModel interface
+	virtual const void* GetCurve() const override;
+
+	virtual void Modify() override;
+
+	virtual void DrawCurve(const FCurveEditor& CurveEditor, const FCurveEditorScreenSpace& ScreenSpace, TArray<TTuple<double, double>>& InterpolatingPoints) const override;
+	virtual void GetKeys(const FCurveEditor& CurveEditor, double MinTime, double MaxTime, double MinValue, double MaxValue, TArray<FKeyHandle>& OutKeyHandles) const override;
+	virtual void AddKeys(TArrayView<const FKeyPosition> InKeyPositions, TArrayView<const FKeyAttributes> InAttributes, TArrayView<TOptional<FKeyHandle>>* OutKeyHandles) override;
+	virtual void RemoveKeys(TArrayView<const FKeyHandle> InKeys) override;
+
+	virtual void GetKeyPositions(TArrayView<const FKeyHandle> InKeys, TArrayView<FKeyPosition> OutKeyPositions) const override;
+	virtual void SetKeyPositions(TArrayView<const FKeyHandle> InKeys, TArrayView<const FKeyPosition> InKeyPositions, EPropertyChangeType::Type ChangeType) override;
+	virtual void GetKeyDrawInfo(ECurvePointType PointType, const FKeyHandle InKeyHandle, FKeyDrawInfo& OutDrawInfo) const override;
+
+	virtual void GetKeyAttributes(TArrayView<const FKeyHandle> InKeys, TArrayView<FKeyAttributes> OutAttributes) const override;
+
+	virtual void GetTimeRange(double& MinTime, double& MaxTime) const override;
+	virtual void GetValueRange(double& MinValue, double& MaxValue) const override;
+
+	virtual int32 GetNumKeys() const override;
+	virtual void GetNeighboringKeys(const FKeyHandle InKeyHandle, TOptional<FKeyHandle>& OutPreviousKeyHandle, TOptional<FKeyHandle>& OutNextKeyHandle) const override;
+
+	virtual bool Evaluate(double ProspectiveTime, double& OutValue) const override;
+
+	virtual bool IsValid() const { return RealCurve != nullptr; }
+
+  private:
+
+  	TWeakObjectPtr<> WeakOwner;
+	FRealCurve* RealCurve;
+};
\ No newline at end of file
diff --git a/Engine/Source/Editor/EditorStyle/Private/StarshipStyle.cpp b/Engine/Source/Editor/EditorStyle/Private/StarshipStyle.cpp
--- a/Engine/Source/Editor/EditorStyle/Private/StarshipStyle.cpp
+++ b/Engine/Source/Editor/EditorStyle/Private/StarshipStyle.cpp
@@ -6530,6 +6530,23 @@ void FStarshipEditorStyle::FStyle::SetupToolkitStyles()
 	// CurveTable Editor
 	{
 		Set( "CurveTableEditor.Tabs.Properties", new IMAGE_BRUSH( "/Icons/icon_tab_SelectionDetails_16x", Icon16x16 ) );
+		Set( "CurveTableEditor.CurveView", new IMAGE_BRUSH("GenericCurveEditor/Icons/GenericCurveEditor_48x", Icon20x20));
+		Set( "CurveTableEditor.TableView", new CORE_IMAGE_BRUSH_SVG("Starship/Common/layout-spreadsheet",  Icon20x20));
+
+		// SEditableTextBox defaults...
+		Set("CurveTableEditor.Cell.Text", FEditableTextBoxStyle()
+			.SetBackgroundImageNormal(FSlateNoResource())
+			.SetBackgroundImageHovered(FSlateRoundedBoxBrush(FStyleColors::Input, 0.0, FStyleColors::Hover, 1.0))
+			.SetBackgroundImageFocused(FSlateRoundedBoxBrush(FStyleColors::Input, 0.0, FStyleColors::Primary, 1.0))
+			.SetBackgroundImageReadOnly(FSlateNoResource())
+			.SetPadding(FMargin(4.f, 5.0f, 4.f, 4.0f)) // The padding should be 4 top, 5 bottom
+			.SetForegroundColor(FStyleColors::Foreground)
+			.SetBackgroundColor(FStyleColors::White)
+			.SetReadOnlyForegroundColor(FSlateColor::UseForeground())
+			.SetFocusedForegroundColor(FStyleColors::White)
+		);
+
+
 	}
 
 	// DataTable Editor
diff --git a/Engine/Source/Editor/UnrealEd/Classes/Factories/CurveTableFactory.h b/Engine/Source/Editor/UnrealEd/Classes/Factories/CurveTableFactory.h
new file mode 100644
--- /dev/null
+++ b/Engine/Source/Editor/UnrealEd/Classes/Factories/CurveTableFactory.h
@@ -0,0 +1,24 @@
+// Copyright Epic Games, Inc. All Rights Reserved.
+
+#pragma once
+
+#include "CoreMinimal.h"
+#include "UObject/ObjectMacros.h"
+#include "Factories/Factory.h"
+#include "CurveTableFactory.generated.h"
+
+class UCurveTable;
+
+/* Creates a CurveTable with Rich Curves */
+
+UCLASS(hidecategories=Object)
+class UNREALED_API UCurveTableFactory : public UFactory
+{
+	GENERATED_UCLASS_BODY()
+
+	// UFactory Interface
+	virtual UObject* FactoryCreateNew(UClass* Class, UObject* InParent, FName Name, EObjectFlags Flags, UObject* Context, FFeedbackContext* Warn) override;
+
+protected:
+	virtual UCurveTable* MakeNewCurveTable(UObject* InParent, FName Name, EObjectFlags Flags);
+};
diff --git a/Engine/Source/Editor/UnrealEd/Private/Factories/CurveTableFactory.cpp b/Engine/Source/Editor/UnrealEd/Private/Factories/CurveTableFactory.cpp
new file mode 100644
--- /dev/null
+++ b/Engine/Source/Editor/UnrealEd/Private/Factories/CurveTableFactory.cpp
@@ -0,0 +1,33 @@
+// Copyright Epic Games, Inc. All Rights Reserved.
+
+#include "Factories/CurveTableFactory.h"
+#include "Engine/CurveTable.h"
+#include "Editor.h"
+
+#define LOCTEXT_NAMESPACE "CurveTableFactory"
+
+UCurveTableFactory::UCurveTableFactory(const FObjectInitializer& ObjectInitializer)
+	: Super(ObjectInitializer)
+{
+	SupportedClass = UCurveTable::StaticClass();
+	bCreateNew = true;
+	bEditAfterNew = true;
+}
+
+UObject* UCurveTableFactory::FactoryCreateNew(UClass* Class, UObject* InParent, FName Name, EObjectFlags Flags, UObject* Context, FFeedbackContext* Warn)
+{
+	UCurveTable* CurveTable = nullptr;
+	if (ensure(SupportedClass == Class))
+	{
+		ensure(0 != (RF_Public & Flags));
+		CurveTable = MakeNewCurveTable(InParent, Name, Flags);
+	}
+	return CurveTable;
+}
+
+UCurveTable* UCurveTableFactory::MakeNewCurveTable(UObject* InParent, FName Name, EObjectFlags Flags)
+{
+	return NewObject<UCurveTable>(InParent, Name, Flags);
+}
+
+#undef LOCTEXT_NAMESPACE 
diff --git a/Engine/Source/Runtime/Slate/Public/Widgets/Input/SNumericEntryBox.h b/Engine/Source/Runtime/Slate/Public/Widgets/Input/SNumericEntryBox.h
--- a/Engine/Source/Runtime/Slate/Public/Widgets/Input/SNumericEntryBox.h
+++ b/Engine/Source/Runtime/Slate/Public/Widgets/Input/SNumericEntryBox.h
@@ -71,6 +71,7 @@ public:
 		, _SpinBoxStyle(&FAppStyle::Get().GetWidgetStyle<FSpinBoxStyle>("NumericEntrySpinBox") )
 		, _Label()
 		, _LabelVAlign( VAlign_Fill )
+		, _Justification(ETextJustify::Left)
 		, _LabelLocation(ELabelLocation::Outside)
 		, _LabelPadding(FMargin(3.f,0.f) )
 		, _BorderForegroundColor(FAppStyle::Get().GetWidgetStyle<FSpinBoxStyle>("NumericEntrySpinBox").ForegroundColor)
@@ -99,6 +100,9 @@ public:
 		SLATE_NAMED_SLOT( FArguments, Label )
 		/** Vertical alignment of the label content */
 		SLATE_ARGUMENT( EVerticalAlignment, LabelVAlign )
+		/** How should the value be justified in the editable text field. */
+		SLATE_ATTRIBUTE(ETextJustify::Type, Justification)
+
 		SLATE_ARGUMENT(ELabelLocation, LabelLocation)
 		/** Padding around the label content */
 		SLATE_ARGUMENT( FMargin, LabelPadding )
@@ -231,6 +235,7 @@ public:
 			.OnTextCommitted(this, &SNumericEntryBox<NumericType>::OnTextCommitted)
 			.SelectAllTextOnCommit(true)
 			.ContextMenuExtender(InArgs._ContextMenuExtender)
+			.Justification(InArgs._Justification)
 			.MinDesiredWidth(InArgs._MinDesiredValueWidth);
 
 		TSharedRef<SOverlay> Overlay = SNew(SOverlay);
