Fixes Crash on Adding curve to curve table

diff --git a/Engine/Source/Editor/CurveTableEditor/Private/CurveTableEditor.cpp b/Engine/Source/Editor/CurveTableEditor/Private/CurveTableEditor.cpp
--- a/Engine/Source/Editor/CurveTableEditor/Private/CurveTableEditor.cpp
+++ b/Engine/Source/Editor/CurveTableEditor/Private/CurveTableEditor.cpp
@@ -54,6 +54,25 @@ struct FCurveTableEditorColumnHeaderData
 	float KeyTime;
 };
 
+namespace {
+
+		FName MakeUniqueCurveName( UCurveTable* Table )
+		{
+				check(Table != nullptr);
+
+				int incr = 0;	
+				FName TestName = FName("Curve", incr);
+
+				const TMap<FName, FRealCurve*>& RowMap = Table->GetRowMap();
+
+				while (RowMap.Contains(TestName))
+				{
+						TestName = FName("Curve", ++incr);
+				}
+
+				return TestName;
+		}
+}
 
 /*
 * FCurveTableEditorItem
@@ -706,14 +725,14 @@ FReply FCurveTableEditor::OnAddCurveClicked()
 
 	if (Table->HasRichCurves())
 	{
-		FName NewCurveUnique = MakeUniqueObjectName(Table, UCurveTable::StaticClass(), "Curve");
+		FName NewCurveUnique = MakeUniqueCurveName(Table);
 		FRichCurve& NewCurve = Table->AddRichCurve(NewCurveUnique);
 		FCurveEditorTreeItem* TreeItem = CurveEditor->AddTreeItem(FCurveEditorTreeItemID());
 		TreeItem->SetStrongItem(MakeShared<FCurveTableEditorItem>(NewCurveUnique, FCurveTableEditorHandle(Table, NewCurveUnique), AvailableColumns));
 	}
 	else
 	{
-		FName NewCurveUnique = MakeUniqueObjectName(Table, UCurveTable::StaticClass(), "Curve");
+		FName NewCurveUnique = MakeUniqueCurveName(Table);
 		FRealCurve& RealCurve = Table->AddSimpleCurve(NewCurveUnique);
 		// Also add a default key for each column 
 		for (auto Column : AvailableColumns)
