Curve Table Editor Undo / Redo

- also fixes memory leak

diff --git a/Engine/Source/Editor/CurveEditor/Private/CurveEditor.cpp b/Engine/Source/Editor/CurveEditor/Private/CurveEditor.cpp
--- a/Engine/Source/Editor/CurveEditor/Private/CurveEditor.cpp
+++ b/Engine/Source/Editor/CurveEditor/Private/CurveEditor.cpp
@@ -269,6 +269,16 @@ const TArray<FCurveEditorTreeItemID>& FCurveEditor::GetRootTreeItems() const
 	return Tree.GetRootItems();
 }
 
+FCurveEditorTreeItemID FCurveEditor::GetTreeIDFromCurveID(FCurveModelID CurveID) const
+{
+	if (TreeIDByCurveID.Contains(CurveID))
+	{
+		return TreeIDByCurveID[CurveID];	
+	}
+
+	return FCurveEditorTreeItemID();
+}
+
 FCurveEditorTreeItem* FCurveEditor::AddTreeItem(FCurveEditorTreeItemID ParentID)
 {
 	return Tree.AddItem(ParentID);
diff --git a/Engine/Source/Editor/CurveEditor/Public/CurveEditor.h b/Engine/Source/Editor/CurveEditor/Public/CurveEditor.h
--- a/Engine/Source/Editor/CurveEditor/Public/CurveEditor.h
+++ b/Engine/Source/Editor/CurveEditor/Public/CurveEditor.h
@@ -320,6 +320,12 @@ public:
 	 */
 	const TArray<FCurveEditorTreeItemID>& GetRootTreeItems() const;
 
+
+	/**
+	 * Find a tree ID id associated with a CurveModelID
+	 */
+	FCurveEditorTreeItemID GetTreeIDFromCurveID(FCurveModelID CurveID) const;
+
 	/**
 	 * Add a new tree item to this curve editor
 	 */
diff --git a/Engine/Source/Editor/CurveTableEditor/Private/CurveTableEditor.cpp b/Engine/Source/Editor/CurveTableEditor/Private/CurveTableEditor.cpp
--- a/Engine/Source/Editor/CurveTableEditor/Private/CurveTableEditor.cpp
+++ b/Engine/Source/Editor/CurveTableEditor/Private/CurveTableEditor.cpp
@@ -4,6 +4,7 @@
 #include "Widgets/Text/STextBlock.h"
 #include "Modules/ModuleManager.h"
 #include "Fonts/FontMeasure.h"
+#include "Framework/Commands/GenericCommands.h"
 #include "Framework/Application/SlateApplication.h"
 #include "Widgets/Layout/SScrollBar.h"
 #include "Framework/Layout/Overscroll.h"
@@ -144,8 +145,12 @@ class FCurveTableEditorItem : public ICurveEditorTreeItem,  public TSharedFromTh
 			return SNew(SNumericEntryBox<float>)
 				.EditableTextBoxStyle( &FAppStyle::Get().GetWidgetStyle<FEditableTextBoxStyle>("CurveTableEditor.Cell.Text") )
 				.Value_Lambda([Curve, KeyHandle] () { return Curve->GetKeyValue(KeyHandle); })
-				// .Value_Lambda([this, InColumnId] () { return CellDataF[InColumnId]; })
-				.OnValueChanged_Lambda([Curve, KeyHandle] (float NewValue) {Curve->SetKeyValue(KeyHandle, NewValue);})
+				.OnValueChanged_Lambda([this, KeyHandle] (float NewValue) 
+				{
+					FScopedTransaction Transaction(LOCTEXT("SetKeyValues", "Set Key Values"));
+					RowHandle.ModifyOwner();
+					RowHandle.GetCurve()->SetKeyValue(KeyHandle, NewValue);
+				})
 				.Justification(ETextJustify::Right)
 			;
 		}
@@ -268,6 +273,8 @@ void FCurveTableEditor::InitCurveTableEditor( const EToolkitMode::Type Mode, con
 	RegenerateMenusAndToolbars();
 
 	FReimportManager::Instance()->OnPostReimport().AddSP(this, &FCurveTableEditor::OnPostReimport);
+
+	GEditor->RegisterForUndo(this);
 }
 
 TSharedRef< FTabManager::FLayout > FCurveTableEditor::InitCurveTableLayout()
@@ -289,6 +296,9 @@ void FCurveTableEditor::BindCommands()
 {
 	FCurveTableEditorCommands::Register();
 
+	ToolkitCommands->MapAction(FGenericCommands::Get().Undo,   FExecuteAction::CreateLambda([]{ GEditor->UndoTransaction(); }));
+	ToolkitCommands->MapAction(FGenericCommands::Get().Redo,   FExecuteAction::CreateLambda([]{ GEditor->RedoTransaction(); }));
+
 	ToolkitCommands->MapAction(FCurveTableEditorCommands::Get().CurveViewToggle,
 		FExecuteAction::CreateSP(this, &FCurveTableEditor::ToggleViewMode),
 		FCanExecuteAction(),
@@ -348,6 +358,16 @@ void FCurveTableEditor::PreChange(const UCurveTable* Changed, FCurveTableEditorU
 {
 }
 
+void FCurveTableEditor::PostUndo(bool bSuccess)
+{
+	RefreshCachedCurveTable();
+}
+
+void FCurveTableEditor::PostRedo(bool bSuccess)
+{
+	RefreshCachedCurveTable();
+}
+
 void FCurveTableEditor::PostChange(const UCurveTable* Changed, FCurveTableEditorUtils::ECurveTableChangeInfo Info)
 {
 	const UCurveTable* Table = GetCurveTable();
@@ -357,9 +377,9 @@ void FCurveTableEditor::PostChange(const UCurveTable* Changed, FCurveTableEditor
 	}
 }
 
-const UCurveTable* FCurveTableEditor::GetCurveTable() const
+UCurveTable* FCurveTableEditor::GetCurveTable() const
 {
-	return Cast<const UCurveTable>(GetEditingObject());
+	return Cast<UCurveTable>(GetEditingObject());
 }
 
 void FCurveTableEditor::HandlePostChange()
@@ -384,6 +404,12 @@ TSharedRef<SDockTab> FCurveTableEditor::SpawnTab_CurveTable( const FSpawnTabArgs
 	FCurveEditorInitParams CurveEditorInitParams;
 	CurveEditor->InitCurveEditor(CurveEditorInitParams);
 
+	// We want this editor to handle undo, not the CurveEditor because
+	// the PostUndo fixes up the selection and in the case of a CurveTable,
+	// the curves have been rebuilt on undo and thus need special handling to restore the selection
+	GEditor->UnregisterForUndo(CurveEditor.Get());
+
+
 	CurveEditorTree = SNew(SCurveEditorTree, CurveEditor.ToSharedRef())
 		.OnTreeViewScrolled(this, &FCurveTableEditor::OnCurveTreeViewScrolled);
 
@@ -507,7 +533,7 @@ void FCurveTableEditor::RefreshTableRowsSelection()
 {
 	if(bUpdatingTableViewSelection == false)
 	{
-		TGuardValue<bool> SelecitonGuard(bUpdatingTableViewSelection, true);
+		TGuardValue<bool> SelectionGuard(bUpdatingTableViewSelection, true);
 
 		TArray<FCurveEditorTreeItemID> CurrentTreeWidgetSelection;
 		TableView->GetSelectedItems(CurrentTreeWidgetSelection);
@@ -532,23 +558,51 @@ void FCurveTableEditor::OnTableViewSelectionChanged(FCurveEditorTreeItemID ItemI
 {
 	if (bUpdatingTableViewSelection == false)
 	{
-		TGuardValue<bool> SelecitonGuard(bUpdatingTableViewSelection, true);
+		TGuardValue<bool> SelectionGuard(bUpdatingTableViewSelection, true);
 		CurveEditor->GetTree()->SetDirectSelection(TableView->GetSelectedItems(), CurveEditor.Get());
 	}
 }
 
 void FCurveTableEditor::RefreshCachedCurveTable()
 {
-
 	// This will trigger to remove any cached widgets in the TableView while we rebuild the model from the source CurveTable
+
+	const TSet<FCurveModelID>& Pinned = CurveEditor->GetPinnedCurves();
+	TSet<FName> PinnedCurves;
+	for (auto PinnedCurveID : Pinned)
+	{
+		FCurveEditorTreeItemID TreeID = CurveEditor->GetTreeIDFromCurveID(PinnedCurveID);
+		if (RowIDMap.Contains(TreeID))
+		{
+			PinnedCurves.Add(RowIDMap[TreeID]);
+		}
+	}
+
+	TSet<FName> SelectedCurves;
+	const TMap<FCurveEditorTreeItemID, ECurveEditorTreeSelectionState>& Selected = CurveEditor->GetTreeSelection();
+	for (const TPair<FCurveEditorTreeItemID, ECurveEditorTreeSelectionState>& SelectionEntry: Selected)
+	{
+		if (SelectionEntry.Value != ECurveEditorTreeSelectionState::None)
+		{
+			if (RowIDMap.Contains(SelectionEntry.Key))
+			{
+				SelectedCurves.Add(RowIDMap[SelectionEntry.Key]);
+			}
+		}
+	}
+
+	// New Selection 
+	TArray<FCurveEditorTreeItemID> NewSelectedItems;
+
 	TableView->SetListItemsSource(EmptyItems);
 	
 	CurveEditor->RemoveAllTreeItems();
 
 	ColumnNamesHeaderRow->ClearColumns();
 	AvailableColumns.Empty();
+	RowIDMap.Empty();
 
-	const UCurveTable* Table = GetCurveTable();
+	UCurveTable* Table = GetCurveTable();
 	if (!Table || Table->GetRowMap().Num() == 0)
 	{
 		return;
@@ -566,6 +620,20 @@ void FCurveTableEditor::RefreshCachedCurveTable()
 			const FName& CurveName = CurveRow.Key;
 			FCurveEditorTreeItem* TreeItem = CurveEditor->AddTreeItem(FCurveEditorTreeItemID());
 			TreeItem->SetStrongItem(MakeShared<FCurveTableEditorItem>(CurveName, FCurveTableEditorHandle(Table, CurveName), AvailableColumns));
+			RowIDMap.Add(TreeItem->GetID(), CurveName);
+
+			if (SelectedCurves.Contains(CurveName))
+			{
+				NewSelectedItems.Add(TreeItem->GetID());
+			}
+
+			if (PinnedCurves.Contains(CurveName))
+			{
+				for (auto ModelID : TreeItem->GetCurves())
+				{
+					CurveEditor->PinCurve(ModelID);
+				}
+			}
 		}
 	}
 
@@ -609,10 +677,28 @@ void FCurveTableEditor::RefreshCachedCurveTable()
 			TSharedPtr<FCurveTableEditorItem> NewItem = MakeShared<FCurveTableEditorItem>(CurveName, FCurveTableEditorHandle(Table, CurveName), AvailableColumns);
 			OnColumnsChanged.AddSP(NewItem.ToSharedRef(), &FCurveTableEditorItem::CacheKeys);
 			TreeItem->SetStrongItem(NewItem);
+			RowIDMap.Add(TreeItem->GetID(), CurveName);
+
+			if (SelectedCurves.Contains(CurveName))
+			{
+				NewSelectedItems.Add(TreeItem->GetID());
+			}
+
+			if (PinnedCurves.Contains(CurveName))
+			{
+				for (auto ModelID : TreeItem->GetOrCreateCurves(CurveEditor.Get()))
+				{
+					CurveEditor->PinCurve(ModelID);
+				}
+			}
 		}
 	}
 
 	TableView->SetListItemsSource(CurveEditorTree->GetSourceItems());
+
+	TGuardValue<bool> SelectionGuard(bUpdatingTableViewSelection, true);
+	CurveEditor->SetTreeSelection(MoveTemp(NewSelectedItems));
+
 }
 
 void FCurveTableEditor::OnCurveTreeViewScrolled(double InScrollOffset)
@@ -720,15 +806,19 @@ TSharedRef<SWidget> FCurveTableEditor::MakeToolbar(TSharedRef<SCurveEditorPanel>
 
 FReply FCurveTableEditor::OnAddCurveClicked()
 {
+	FScopedTransaction Transaction(LOCTEXT("AddCurve", "Add Curve"));
+
 	UCurveTable* Table = Cast<UCurveTable>(GetEditingObject());
 	check(Table != nullptr);
 
+	Table->Modify();
 	if (Table->HasRichCurves())
 	{
 		FName NewCurveUnique = MakeUniqueCurveName(Table);
 		FRichCurve& NewCurve = Table->AddRichCurve(NewCurveUnique);
 		FCurveEditorTreeItem* TreeItem = CurveEditor->AddTreeItem(FCurveEditorTreeItemID());
 		TreeItem->SetStrongItem(MakeShared<FCurveTableEditorItem>(NewCurveUnique, FCurveTableEditorHandle(Table, NewCurveUnique), AvailableColumns));
+		RowIDMap.Add(TreeItem->GetID(), NewCurveUnique);
 	}
 	else
 	{
@@ -744,7 +834,7 @@ FReply FCurveTableEditor::OnAddCurveClicked()
 		TSharedPtr<FCurveTableEditorItem> NewItem = MakeShared<FCurveTableEditorItem>(NewCurveUnique, FCurveTableEditorHandle(Table, NewCurveUnique), AvailableColumns);
 		OnColumnsChanged.AddSP(NewItem.ToSharedRef(), &FCurveTableEditorItem::CacheKeys);
 		TreeItem->SetStrongItem(NewItem);
-
+		RowIDMap.Add(TreeItem->GetID(), NewCurveUnique);
 
 	}
 
@@ -785,6 +875,9 @@ void FCurveTableEditor::AddNewKeyColumn(float NewKeyTime)
 
 	if (!Table->HasRichCurves())
 	{
+		FScopedTransaction Transaction(LOCTEXT("AddKeyColumn", "AddKeyColumn"));
+		Table->Modify();	
+
 		// Make sure we don't already have a key at this time
 
 		// 1. Add new keys to every curve
diff --git a/Engine/Source/Editor/CurveTableEditor/Private/CurveTableEditor.h b/Engine/Source/Editor/CurveTableEditor/Private/CurveTableEditor.h
--- a/Engine/Source/Editor/CurveTableEditor/Private/CurveTableEditor.h
+++ b/Engine/Source/Editor/CurveTableEditor/Private/CurveTableEditor.h
@@ -37,6 +37,7 @@ typedef TSharedPtr<FCurveTableEditorColumnHeaderData> FCurveTableEditorColumnHea
 class FCurveTableEditor :
 	public ICurveTableEditor
 	, public FCurveTableEditorUtils::INotifyOnCurveTableChanged
+	, public FEditorUndoClient
 {
 
 public:
@@ -64,9 +65,11 @@ public:
 	// INotifyOnDataTableChanged
 	virtual void PreChange(const UCurveTable* Changed, FCurveTableEditorUtils::ECurveTableChangeInfo Info) override;
 	virtual void PostChange(const UCurveTable* Changed, FCurveTableEditorUtils::ECurveTableChangeInfo Info) override;
+	virtual void PostUndo(bool bSuccess) override;
+	virtual void PostRedo(bool bSuccess) override;
 
 	/** Get the curve table being edited */
-	const UCurveTable* GetCurveTable() const;
+	UCurveTable* GetCurveTable() const;
 
 	void HandlePostChange();
 
@@ -167,4 +170,7 @@ protected:
 	/** An empty source list used to initialize or when rebuilding the TableView */
 	TArray<FCurveEditorTreeItemID> EmptyItems;
 
+	/** Associates the CurveEditor Ids with the FName Ids from the Curve Table */
+	TMap< FCurveEditorTreeItemID, FName > RowIDMap;
+
 };
diff --git a/Engine/Source/Editor/CurveTableEditor/Private/CurveTableEditorHandle.cpp b/Engine/Source/Editor/CurveTableEditor/Private/CurveTableEditorHandle.cpp
--- a/Engine/Source/Editor/CurveTableEditor/Private/CurveTableEditorHandle.cpp
+++ b/Engine/Source/Editor/CurveTableEditor/Private/CurveTableEditorHandle.cpp
@@ -55,14 +55,32 @@ TArray<FRichCurveEditInfo> FCurveTableEditorHandle::GetCurves()
 	return Curves;
 }
 
+TArray<const UObject*> FCurveTableEditorHandle::GetOwners() const
+{ 
+	TArray<const UObject*> Owners;
+	if (CurveTable != nullptr)
+	{
+		Owners.Add(CurveTable.Get());
+	}
+
+	return Owners;
+}
+
 void FCurveTableEditorHandle::ModifyOwner()
 {
-	check(false);	// This handle is read only, so cannot be used to modify curves
+	if (CurveTable != nullptr && RowName != NAME_None)
+	{
+		CurveTable->Modify();
+	}
+
 }
 
 void FCurveTableEditorHandle::MakeTransactional()
 {
-	check(false);	// This handle is read only, so cannot be used to modify curves
+	if (CurveTable != nullptr)
+	{
+		CurveTable->SetFlags(CurveTable->GetFlags() | RF_Transactional);
+	}
 }
 
 void FCurveTableEditorHandle::OnCurveChanged(const TArray<FRichCurveEditInfo>& ChangedCurveEditInfos)
diff --git a/Engine/Source/Editor/CurveTableEditor/Private/CurveTableEditorHandle.h b/Engine/Source/Editor/CurveTableEditor/Private/CurveTableEditorHandle.h
--- a/Engine/Source/Editor/CurveTableEditor/Private/CurveTableEditorHandle.h
+++ b/Engine/Source/Editor/CurveTableEditor/Private/CurveTableEditorHandle.h
@@ -17,13 +17,13 @@ struct FCurveTableEditorHandle : public FCurveOwnerInterface
 		, RowName(NAME_None)
 	{ }
 
-	FCurveTableEditorHandle(const UCurveTable* InCurveTable, FName InRowName)
+	FCurveTableEditorHandle(UCurveTable* InCurveTable, FName InRowName)
 		: CurveTable(InCurveTable)
 		, RowName(InRowName)
 	{ }
 
 	/** Pointer to table we want a row from */
-	TWeakObjectPtr<const UCurveTable> CurveTable;
+	TWeakObjectPtr<UCurveTable> CurveTable;
 
 	/** Name of row in the table that we want */
 	FName RowName;
@@ -35,11 +35,8 @@ struct FCurveTableEditorHandle : public FCurveOwnerInterface
 	virtual void MakeTransactional() override;
 	virtual void OnCurveChanged(const TArray<FRichCurveEditInfo>& ChangedCurveEditInfos) override;
 	virtual bool IsValidCurve(FRichCurveEditInfo CurveInfo) override;
-	virtual TArray<const UObject*> GetOwners() const override
-	{ 
-		//Note this is read only so we return nothing
-		return TArray<const UObject*>(); 
-	}
+	virtual TArray<const UObject*> GetOwners() const override;
+	
 
 	//~ End FCurveOwnerInterface Interface.
 
diff --git a/Engine/Source/Runtime/Engine/Private/CurveTable.cpp b/Engine/Source/Runtime/Engine/Private/CurveTable.cpp
--- a/Engine/Source/Runtime/Engine/Private/CurveTable.cpp
+++ b/Engine/Source/Runtime/Engine/Private/CurveTable.cpp
@@ -119,7 +119,11 @@ void UCurveTable::Serialize(FArchive& Ar)
 		}
 
 		bool bCouldConvertToSimpleCurves = bUpgradingCurveTable;
-		RowMap.Reserve(NumRows);
+
+		// copy any previous curves to free after replacing.
+		TMap<FName, FRealCurve*> TempMap(RowMap);
+
+		RowMap.Empty(NumRows);
 		for (int32 RowIdx = 0; RowIdx < NumRows; RowIdx++)
 		{
 			// Load row name
@@ -184,6 +188,11 @@ void UCurveTable::Serialize(FArchive& Ar)
 				Curve.Value = NewCurve;
 			}
 		}
+
+		for (TPair<FName, FRealCurve*>& Curve : TempMap)
+		{
+			delete Curve.Value;
+		}
 	}
 	else if (Ar.IsSaving())
 	{
@@ -966,6 +975,7 @@ void UCurveTable::MakeTransactional()
 void UCurveTable::OnCurveChanged(const TArray<FRichCurveEditInfo>& ChangedCurveEditInfos)
 {
 	CURVETABLE_CHANGE_SCOPE();
+	OnCurveTableChanged().Broadcast();
 }
 
 bool UCurveTable::IsValidCurve(FRichCurveEditInfo CurveInfo)
