Curve Table Editor - Constant Keys Support

diff --git a/Engine/Source/Editor/CurveTableEditor/Private/CurveTableEditor.cpp b/Engine/Source/Editor/CurveTableEditor/Private/CurveTableEditor.cpp
--- a/Engine/Source/Editor/CurveTableEditor/Private/CurveTableEditor.cpp
+++ b/Engine/Source/Editor/CurveTableEditor/Private/CurveTableEditor.cpp
@@ -670,8 +670,16 @@ void FCurveTableEditor::RefreshCachedCurveTable()
 		}
 
 		// Setup the CurveEditorTree 
-		for (const TPair<FName, FRealCurve*>& CurveRow : Table->GetRowMap())
+
+		// Store the default Interpolation Mode
+		InterpMode = RCIM_None;
+		for (const TPair<FName, FSimpleCurve*>& CurveRow : Table->GetSimpleCurveRowMap())
 		{
+			if (InterpMode == RCIM_None) 
+			{
+				InterpMode = CurveRow.Value->GetKeyInterpMode();
+			}
+
 			const FName& CurveName = CurveRow.Key;
 			FCurveEditorTreeItem* TreeItem = CurveEditor->AddTreeItem(FCurveEditorTreeItemID());
 			TSharedPtr<FCurveTableEditorItem> NewItem = MakeShared<FCurveTableEditorItem>(CurveName, FCurveTableEditorHandle(Table, CurveName), AvailableColumns);
@@ -823,7 +831,9 @@ FReply FCurveTableEditor::OnAddCurveClicked()
 	else
 	{
 		FName NewCurveUnique = MakeUniqueCurveName(Table);
-		FRealCurve& RealCurve = Table->AddSimpleCurve(NewCurveUnique);
+		FSimpleCurve& RealCurve = Table->AddSimpleCurve(NewCurveUnique);
+		RealCurve.SetKeyInterpMode(InterpMode);
+
 		// Also add a default key for each column 
 		for (auto Column : AvailableColumns)
 		{
diff --git a/Engine/Source/Editor/CurveTableEditor/Private/CurveTableEditor.h b/Engine/Source/Editor/CurveTableEditor/Private/CurveTableEditor.h
--- a/Engine/Source/Editor/CurveTableEditor/Private/CurveTableEditor.h
+++ b/Engine/Source/Editor/CurveTableEditor/Private/CurveTableEditor.h
@@ -173,4 +173,7 @@ protected:
 	/** Associates the CurveEditor Ids with the FName Ids from the Curve Table */
 	TMap< FCurveEditorTreeItemID, FName > RowIDMap;
 
+	/** The default interpolation type of the curves within the table */
+	ERichCurveInterpMode InterpMode;
+
 };
diff --git a/Engine/Source/Editor/CurveTableEditor/Private/RealCurveModel.cpp b/Engine/Source/Editor/CurveTableEditor/Private/RealCurveModel.cpp
--- a/Engine/Source/Editor/CurveTableEditor/Private/RealCurveModel.cpp
+++ b/Engine/Source/Editor/CurveTableEditor/Private/RealCurveModel.cpp
@@ -5,6 +5,7 @@
 
 #include "Curves/RealCurve.h"
 #include "Styling/AppStyle.h"
+#include "CurveEditorScreenSpace.h"
 
 FRealCurveModel::FRealCurveModel(FRealCurve* InRealCurve, UObject* InOwner)
 	: WeakOwner(InOwner)
@@ -35,10 +36,37 @@ void FRealCurveModel::DrawCurve(const FCurveEditor& CurveEditor, const FCurveEdi
 
 	if (IsValid())
 	{
+
+		FKeyHandle LastKeyHandle = RealCurve->GetFirstKeyHandle();
+		double LastKeyTime =  double(ScreenSpace.GetInputMin());
+		double LastKeyValue =  double(RealCurve->GetKeyValue(LastKeyHandle));
+		double LastInterpMode = RealCurve->GetKeyInterpMode(LastKeyHandle);
+
+		if (LastInterpMode == RCIM_Constant)
+		{
+			InterpolatingPoints.Add(MakeTuple( LastKeyTime, LastKeyValue));
+		}
+
 		for (auto It = RealCurve->GetKeyHandleIterator(); It; ++It)
 		{
 			auto KeyPair = RealCurve->GetKeyTimeValuePair(*It);
+
+			// if  constant , add another point to mark the end of the previous' key reign
+			if (LastInterpMode == RCIM_Constant)
+			{
+				InterpolatingPoints.Add(MakeTuple( double(KeyPair.Key), LastKeyValue));
+			}
+
 			InterpolatingPoints.Add(MakeTuple( double(KeyPair.Key), double(KeyPair.Value)));
+
+			LastKeyHandle = RealCurve->GetNextKey(LastKeyHandle);
+			LastKeyValue = KeyPair.Value;
+			LastInterpMode = RealCurve->GetKeyInterpMode(LastKeyHandle);
+		}
+
+		if (LastInterpMode == RCIM_Constant)
+		{
+			InterpolatingPoints.Add(MakeTuple( double(ScreenSpace.GetInputMax()), LastKeyValue));
 		}
 	}
 }
