nonunity fix

diff --git a/Engine/Source/Editor/CurveTableEditor/Private/CurveTableEditor.h b/Engine/Source/Editor/CurveTableEditor/Private/CurveTableEditor.h
--- a/Engine/Source/Editor/CurveTableEditor/Private/CurveTableEditor.h
+++ b/Engine/Source/Editor/CurveTableEditor/Private/CurveTableEditor.h
@@ -13,6 +13,7 @@
 #include "Tree/ICurveEditorTreeItem.h"
 #include "CurveTableEditorHandle.h"
 #include "CurveTableEditorUtils.h"
+#include "EditorUndoClient.h"
 
 #include "CurveEditorTypes.h"
 class FCurveEditor;
