Curve Table Editor Updates

- Switched to using an AssetEditorToolkit Toolbar
- Added support to rename curves

diff --git a/Engine/Source/Editor/CurveTableEditor/Private/CurveTableEditor.cpp b/Engine/Source/Editor/CurveTableEditor/Private/CurveTableEditor.cpp
--- a/Engine/Source/Editor/CurveTableEditor/Private/CurveTableEditor.cpp
+++ b/Engine/Source/Editor/CurveTableEditor/Private/CurveTableEditor.cpp
@@ -1,18 +1,19 @@
 // Copyright Epic Games, Inc. All Rights Reserved.
 
 #include "CurveTableEditor.h"
-#include "Widgets/Text/STextBlock.h"
 #include "Modules/ModuleManager.h"
+#include "Engine/CurveTable.h"
 #include "Fonts/FontMeasure.h"
 #include "Framework/Commands/GenericCommands.h"
 #include "Framework/Application/SlateApplication.h"
-#include "Widgets/Layout/SScrollBar.h"
 #include "Framework/Layout/Overscroll.h"
 #include "Widgets/Input/SButton.h"
 #include "Widgets/Input/SNumericEntryBox.h"
 #include "Widgets/Input/SSegmentedControl.h"
-#include "Widgets/Views/SListView.h"
+#include "Widgets/Layout/SScrollBar.h"
 #include "Widgets/Layout/SScrollBox.h"
+#include "Widgets/Text/SInlineEditableTextBlock.h"
+#include "Widgets/Views/SListView.h"
 #include "SPositiveActionButton.h"
 #include "EditorStyleSet.h"
 #include "Styling/StyleColors.h"
@@ -28,6 +29,7 @@
 #include "Tree/CurveEditorTreeFilter.h"
 
 #include "Tree/SCurveEditorTreeTextFilter.h"
+#include "SSimpleButton.h"
 
 #include "RealCurveModel.h"
 #include "RichCurveEditorModel.h"
@@ -97,8 +99,10 @@ class FCurveTableEditorItem : public ICurveEditorTreeItem,  public TSharedFromTh
   	};
 
   public: 
-	FCurveTableEditorItem (const FName& InRowId, FCurveTableEditorHandle InRowHandle, const TArray<FCurveTableEditorColumnHeaderDataPtr>& InColumns)
-		: RowId(InRowId)
+	FCurveTableEditorItem (TWeakPtr<FCurveTableEditor> InCurveTableEditor, const FCurveEditorTreeItemID& InTreeID, const FName& InRowId, FCurveTableEditorHandle InRowHandle, const TArray<FCurveTableEditorColumnHeaderDataPtr>& InColumns)
+		: CurveTableEditor(InCurveTableEditor)
+		, TreeID(InTreeID)
+		, RowId(InRowId)
 		, RowHandle(InRowHandle)
 		, Columns(InColumns)
 	{
@@ -118,9 +122,11 @@ class FCurveTableEditorItem : public ICurveEditorTreeItem,  public TSharedFromTh
 				.HAlign(HAlign_Right)
 				.AutoWidth()
 				[
-					SNew(STextBlock)
+					SAssignNew(InlineRenameWidget, SInlineEditableTextBlock)
 					.Text(DisplayName)
 					.ColorAndOpacity(FSlateColor::UseForeground())
+					.OnTextCommitted(this, &FCurveTableEditorItem::HandleNameCommitted)
+					.OnVerifyTextChanged(this, &FCurveTableEditorItem::VerifyNameChanged)
 				];
 		}
 		else if (InColumnName == ColumnNames.SelectHeader)
@@ -223,6 +229,71 @@ class FCurveTableEditorItem : public ICurveEditorTreeItem,  public TSharedFromTh
 		}
 	}
 
+	void EnterRenameMode()
+	{
+		InlineRenameWidget->EnterEditingMode();
+	}
+
+	bool VerifyNameChanged(const FText& InText, FText& OutErrorMessage)
+	{
+		FName CheckName = FName(*InText.ToString());
+		if (CheckName == RowId)
+		{
+			return true;	
+		}
+
+		if (RowHandle.CurveTable.IsValid())
+		{
+			UCurveTable* Table = RowHandle.CurveTable.Get();
+			const TMap<FName, FRealCurve*>& RowMap = Table->GetRowMap();
+			if (RowMap.Contains(CheckName))
+			{
+
+				OutErrorMessage = LOCTEXT("NameAlreadyUsed", "Row Names Must Be Unique");
+				return false;
+			}
+			return true;
+		}
+		return false;
+	}
+
+	void HandleNameCommitted(const FText& CommittedText, ETextCommit::Type CommitInfo)
+	{
+		if (CommitInfo == ETextCommit::OnEnter)
+		{
+			TSharedPtr<FCurveTableEditor> TableEditorPtr = CurveTableEditor.Pin();
+			if (TableEditorPtr != nullptr)
+			{
+				FName OldName = RowId;
+				FName NewName = *CommittedText.ToString();
+
+				DisplayName = CommittedText;
+				InlineRenameWidget->SetText(DisplayName);
+
+				RowHandle.RowName = NewName;
+				RowId = NewName;
+
+				TableEditorPtr->HandleCurveRename(TreeID, OldName, NewName);
+
+				TSharedPtr<FCurveEditor> CurveEditor = TableEditorPtr->GetCurveEditor();
+				FCurveEditorTreeItem& TreeItem = CurveEditor->GetTreeItem(TreeID);
+				for (FCurveModelID ModelID : TreeItem.GetCurves())
+				{
+					if (FCurveModel* CurveModel = CurveEditor->FindCurve(ModelID))
+					{
+						CurveModel->SetShortDisplayName(DisplayName);
+					}
+				}
+			}
+		}
+	}
+
+	/** Hold onto a weak ptr to the CurveTableEditor specifically for deleting and renaming  */
+	TWeakPtr<FCurveTableEditor> CurveTableEditor;
+
+	/** The CurveEditor's Unique ID for the TreeItem this item is attached to (SetStrongItem) */
+	FCurveEditorTreeItemID TreeID;
+
 	/** Unique ID used to identify this row */
 	FName RowId;
 
@@ -237,6 +308,10 @@ class FCurveTableEditorItem : public ICurveEditorTreeItem,  public TSharedFromTh
 
 	/** A Reference to the available columns in the TableView */
 	const TArray<FCurveTableEditorColumnHeaderDataPtr>& Columns;
+
+	/** Inline editable text box for renaming */
+	TSharedPtr<SInlineEditableTextBlock> InlineRenameWidget;
+
 };
 
 
@@ -270,6 +345,7 @@ void FCurveTableEditor::InitCurveTableEditor( const EToolkitMode::Type Mode, con
 	
 	BindCommands();
 	ExtendMenu();
+	ExtendToolbar();
 	RegenerateMenusAndToolbars();
 
 	FReimportManager::Instance()->OnPostReimport().AddSP(this, &FCurveTableEditor::OnPostReimport);
@@ -299,10 +375,17 @@ void FCurveTableEditor::BindCommands()
 	ToolkitCommands->MapAction(FGenericCommands::Get().Undo,   FExecuteAction::CreateLambda([]{ GEditor->UndoTransaction(); }));
 	ToolkitCommands->MapAction(FGenericCommands::Get().Redo,   FExecuteAction::CreateLambda([]{ GEditor->RedoTransaction(); }));
 
-	ToolkitCommands->MapAction(FCurveTableEditorCommands::Get().CurveViewToggle,
+	ToolkitCommands->MapAction(
+		FCurveTableEditorCommands::Get().CurveViewToggle,
 		FExecuteAction::CreateSP(this, &FCurveTableEditor::ToggleViewMode),
 		FCanExecuteAction(),
-		FIsActionChecked::CreateSP(this, &FCurveTableEditor::IsCurveViewChecked));
+		FIsActionChecked::CreateSP(this, &FCurveTableEditor::IsCurveViewChecked)
+	);
+
+	ToolkitCommands->MapAction(
+		FCurveTableEditorCommands::Get().AppendKeyColumn,
+		FExecuteAction::CreateSP(this, &FCurveTableEditor::OnAddNewKeyColumn)
+	);
 }
 
 void FCurveTableEditor::ExtendMenu()
@@ -318,7 +401,7 @@ void FCurveTableEditor::ExtendMenu()
 				MenuBuilder.AddMenuEntry(FCurveTableEditorCommands::Get().CurveViewToggle);
 			}
 			MenuBuilder.EndSection();
-		}
+			}
 	};
 
 	MenuExtender->AddMenuExtension(
@@ -334,6 +417,63 @@ void FCurveTableEditor::ExtendMenu()
 	AddMenuExtender(CurveTableEditorModule.GetMenuExtensibilityManager()->GetAllExtenders(GetToolkitCommands(), GetEditingObjects()));
 }
 
+void FCurveTableEditor::ExtendToolbar()
+{
+	ToolbarExtender = MakeShareable(new FExtender);
+
+	ToolbarExtender->AddToolBarExtension(
+		"Asset",
+		EExtensionHook::After,
+		GetToolkitCommands(),
+		FToolBarExtensionDelegate::CreateLambda([this](FToolBarBuilder& ParentToolbarBuilder)
+		{
+
+			ParentToolbarBuilder.BeginSection("CurveTable");
+
+			bool HasRichCurves = GetCurveTable()->HasRichCurves();
+			ParentToolbarBuilder.AddWidget(
+				SNew(SSegmentedControl<ECurveTableViewMode>)
+				.Visibility(HasRichCurves ? EVisibility::Collapsed : EVisibility::Visible)
+				.OnValueChanged_Lambda([this] (ECurveTableViewMode InMode) {if (InMode != GetViewMode()) ToggleViewMode();  } )
+				.Value(this, &FCurveTableEditor::GetViewMode)
+
+				+SSegmentedControl<ECurveTableViewMode>::Slot(ECurveTableViewMode::CurveTable)
+			    .Icon(FAppStyle::Get().GetBrush("CurveTableEditor.CurveView"))
+
+				+SSegmentedControl<ECurveTableViewMode>::Slot(ECurveTableViewMode::Grid)
+			    .Icon(FAppStyle::Get().GetBrush("CurveTableEditor.TableView"))
+			);
+
+			if (InterpMode == RCIM_Constant)
+			{
+				ParentToolbarBuilder.AddToolBarButton(
+					FCurveTableEditorCommands::Get().AppendKeyColumn,
+					NAME_None, 
+					FText::GetEmpty(),
+					TAttribute<FText>(), 
+					FSlateIcon(FAppStyle::Get().GetStyleSetName(), "Sequencer.KeySquare")
+				);
+			}
+
+			if (InterpMode == RCIM_Linear)
+			{
+				ParentToolbarBuilder.AddToolBarButton(
+					FCurveTableEditorCommands::Get().AppendKeyColumn,
+					NAME_None, 
+					FText::GetEmpty(),
+					TAttribute<FText>(), 
+					FSlateIcon(FAppStyle::Get().GetStyleSetName(), "Sequencer.KeyTriangle")
+				);
+			}
+
+			ParentToolbarBuilder.EndSection();
+		})
+	);
+
+	AddToolbarExtender(ToolbarExtender);
+
+}
+
 FName FCurveTableEditor::GetToolkitFName() const
 {
 	return FName("CurveTableEditor");
@@ -411,7 +551,8 @@ TSharedRef<SDockTab> FCurveTableEditor::SpawnTab_CurveTable( const FSpawnTabArgs
 
 
 	CurveEditorTree = SNew(SCurveEditorTree, CurveEditor.ToSharedRef())
-		.OnTreeViewScrolled(this, &FCurveTableEditor::OnCurveTreeViewScrolled);
+		.OnTreeViewScrolled(this, &FCurveTableEditor::OnCurveTreeViewScrolled)
+		.OnMouseButtonDoubleClick(this, &FCurveTableEditor::OnRequestCurveRename);
 
 	TSharedRef<SCurveEditorPanel> CurveEditorPanel = SNew(SCurveEditorPanel, CurveEditor.ToSharedRef());
 
@@ -614,12 +755,13 @@ void FCurveTableEditor::RefreshCachedCurveTable()
 
 	if (Table->HasRichCurves())
 	{
+		InterpMode = RCIM_Cubic;
 		for (const TPair<FName, FRichCurve*>& CurveRow : Table->GetRichCurveRowMap())
 		{
 			// Setup the CurveEdtiorTree
 			const FName& CurveName = CurveRow.Key;
 			FCurveEditorTreeItem* TreeItem = CurveEditor->AddTreeItem(FCurveEditorTreeItemID());
-			TreeItem->SetStrongItem(MakeShared<FCurveTableEditorItem>(CurveName, FCurveTableEditorHandle(Table, CurveName), AvailableColumns));
+			TreeItem->SetStrongItem(MakeShared<FCurveTableEditorItem>(SharedThis(this), TreeItem->GetID(), CurveName, FCurveTableEditorHandle(Table, CurveName), AvailableColumns));
 			RowIDMap.Add(TreeItem->GetID(), CurveName);
 
 			if (SelectedCurves.Contains(CurveName))
@@ -682,7 +824,7 @@ void FCurveTableEditor::RefreshCachedCurveTable()
 
 			const FName& CurveName = CurveRow.Key;
 			FCurveEditorTreeItem* TreeItem = CurveEditor->AddTreeItem(FCurveEditorTreeItemID());
-			TSharedPtr<FCurveTableEditorItem> NewItem = MakeShared<FCurveTableEditorItem>(CurveName, FCurveTableEditorHandle(Table, CurveName), AvailableColumns);
+			TSharedPtr<FCurveTableEditorItem> NewItem = MakeShared<FCurveTableEditorItem>(SharedThis(this), TreeItem->GetID(), CurveName, FCurveTableEditorHandle(Table, CurveName), AvailableColumns);
 			OnColumnsChanged.AddSP(NewItem.ToSharedRef(), &FCurveTableEditorItem::CacheKeys);
 			TreeItem->SetStrongItem(NewItem);
 			RowIDMap.Add(TreeItem->GetID(), CurveName);
@@ -753,6 +895,7 @@ bool FCurveTableEditor::IsCurveViewChecked() const
 
 TSharedRef<SWidget> FCurveTableEditor::MakeToolbar(TSharedRef<SCurveEditorPanel>& InEditorPanel)
 {
+
 	FToolBarBuilder ToolBarBuilder(InEditorPanel->GetCommands(), FMultiBoxCustomization::None, InEditorPanel->GetToolbarExtender(), true);
 	ToolBarBuilder.SetStyle(&FAppStyle::Get(), "Sequencer.ToolBar");
 	ToolBarBuilder.BeginSection("Asset");
@@ -763,43 +906,6 @@ TSharedRef<SWidget> FCurveTableEditor::MakeToolbar(TSharedRef<SCurveEditorPanel>
 
 	return SNew(SHorizontalBox)
 
-	+SHorizontalBox::Slot()
-	.AutoWidth()
-	.VAlign(VAlign_Center)
-	.HAlign(HAlign_Center)
-	.Padding(FMargin(2.0, 4.0, 8.f, 4.f))
-	[
-		SNew(SSegmentedControl<ECurveTableViewMode>)
-		.Visibility(HasRichCurves ? EVisibility::Collapsed : EVisibility::Visible)
-		.OnValueChanged_Lambda([this] (ECurveTableViewMode InMode) {if (InMode != GetViewMode()) ToggleViewMode();  } )
-		.Value(this, &FCurveTableEditor::GetViewMode)
-
-		+SSegmentedControl<ECurveTableViewMode>::Slot(ECurveTableViewMode::CurveTable)
-        .Icon(FAppStyle::Get().GetBrush("CurveTableEditor.CurveView"))
-
-		+SSegmentedControl<ECurveTableViewMode>::Slot(ECurveTableViewMode::Grid)
-        .Icon(FAppStyle::Get().GetBrush("CurveTableEditor.TableView"))
-	]
-
-	+SHorizontalBox::Slot()
-	.AutoWidth()
-	.VAlign(VAlign_Center)
-	[
-		SNew(SBox)
-		[
-			SNew(SButton)
-            .ButtonStyle( &FAppStyle::Get().GetWidgetStyle< FButtonStyle >( "SimpleButton" ) )
-			.Visibility(this, &FCurveTableEditor::GetTableViewControlsVisibility)
-			.OnClicked(this, &FCurveTableEditor::OnAddNewKeyColumn)
-			.ToolTipText(LOCTEXT("CurveTableEditor_AddKeyColumnTooltip", "Append a new column to the curve table.\nEvery Curve or Table Row will have a new key appended."))
-			[
-                SNew(SImage)
-                .ColorAndOpacity(FSlateColor::UseForeground())
-                .Image(FAppStyle::Get().GetBrush("Sequencer.KeyTriangle")) 
-	        ]
-		]
-	]
-
 	+SHorizontalBox::Slot()
 	.AutoWidth()
 	.VAlign(VAlign_Center)
@@ -825,7 +931,7 @@ FReply FCurveTableEditor::OnAddCurveClicked()
 		FName NewCurveUnique = MakeUniqueCurveName(Table);
 		FRichCurve& NewCurve = Table->AddRichCurve(NewCurveUnique);
 		FCurveEditorTreeItem* TreeItem = CurveEditor->AddTreeItem(FCurveEditorTreeItemID());
-		TreeItem->SetStrongItem(MakeShared<FCurveTableEditorItem>(NewCurveUnique, FCurveTableEditorHandle(Table, NewCurveUnique), AvailableColumns));
+		TreeItem->SetStrongItem(MakeShared<FCurveTableEditorItem>(SharedThis(this), TreeItem->GetID(), NewCurveUnique, FCurveTableEditorHandle(Table, NewCurveUnique), AvailableColumns));
 		RowIDMap.Add(TreeItem->GetID(), NewCurveUnique);
 	}
 	else
@@ -841,7 +947,7 @@ FReply FCurveTableEditor::OnAddCurveClicked()
 		}
 
 		FCurveEditorTreeItem* TreeItem = CurveEditor->AddTreeItem(FCurveEditorTreeItemID());
-		TSharedPtr<FCurveTableEditorItem> NewItem = MakeShared<FCurveTableEditorItem>(NewCurveUnique, FCurveTableEditorHandle(Table, NewCurveUnique), AvailableColumns);
+		TSharedPtr<FCurveTableEditorItem> NewItem = MakeShared<FCurveTableEditorItem>(SharedThis(this), TreeItem->GetID(), NewCurveUnique, FCurveTableEditorHandle(Table, NewCurveUnique), AvailableColumns);
 		OnColumnsChanged.AddSP(NewItem.ToSharedRef(), &FCurveTableEditorItem::CacheKeys);
 		TreeItem->SetStrongItem(NewItem);
 		RowIDMap.Add(TreeItem->GetID(), NewCurveUnique);
@@ -851,7 +957,7 @@ FReply FCurveTableEditor::OnAddCurveClicked()
 	return FReply::Handled();
 }
 
-FReply FCurveTableEditor::OnAddNewKeyColumn()
+void FCurveTableEditor::OnAddNewKeyColumn()
 {
 	UCurveTable* Table = Cast<UCurveTable>(GetEditingObject());
 	check(Table != nullptr);
@@ -874,8 +980,6 @@ FReply FCurveTableEditor::OnAddNewKeyColumn()
 
 		AddNewKeyColumn(NewKeyTime);
 	}
-
-	return FReply::Handled();
 }
 
 void FCurveTableEditor::AddNewKeyColumn(float NewKeyTime)
@@ -920,4 +1024,37 @@ void FCurveTableEditor::AddNewKeyColumn(float NewKeyTime)
 	}
 }
 
-#undef LOCTEXT_NAMESPACE
+void FCurveTableEditor::OnRequestCurveRename(FCurveEditorTreeItemID TreeItemId)
+{
+	const FCurveEditorTreeItem* TreeItem = CurveEditor->FindTreeItem(TreeItemId);
+	if (TreeItem != nullptr)
+	{
+		TSharedPtr<ICurveEditorTreeItem> CurveEditorTreeItem = TreeItem->GetItem();
+		if (CurveEditorTreeItem.IsValid())
+		{
+			TSharedPtr<FCurveTableEditorItem> CurveTableEditorItem = StaticCastSharedPtr<FCurveTableEditorItem>(CurveEditorTreeItem);
+			CurveTableEditorItem->EnterRenameMode();
+		}
+	}
+}
+
+void FCurveTableEditor::HandleCurveRename(FCurveEditorTreeItemID& TreeID, FName& CurrentCurve, FName& NewCurveName)
+{
+	// Update the underlying Curve Data Asset itself 
+	UCurveTable* Table = Cast<UCurveTable>(GetEditingObject());
+	check(Table != nullptr);
+
+	FScopedTransaction Transaction(LOCTEXT("RenameCurve", "Rename Curve"));
+	Table->SetFlags(RF_Transactional);
+	Table->Modify();
+	Table->RenameRow(CurrentCurve, NewCurveName);
+
+	FPropertyChangedEvent PropertyChangeStruct(nullptr, EPropertyChangeType::ValueSet);
+	Table->PostEditChangeProperty(PropertyChangeStruct);
+
+	// Update our internal map of TreeIDs to FNames
+	RowIDMap[TreeID] = NewCurveName;
+
+}
+
+#undef LOCTEXT_NAMESPACE
\ No newline at end of file
diff --git a/Engine/Source/Editor/CurveTableEditor/Private/CurveTableEditor.h b/Engine/Source/Editor/CurveTableEditor/Private/CurveTableEditor.h
--- a/Engine/Source/Editor/CurveTableEditor/Private/CurveTableEditor.h
+++ b/Engine/Source/Editor/CurveTableEditor/Private/CurveTableEditor.h
@@ -72,6 +72,9 @@ public:
 	/** Get the curve table being edited */
 	UCurveTable* GetCurveTable() const;
 
+	/** Get the curve editor used for the Curve View*/
+	TSharedPtr<FCurveEditor> GetCurveEditor() const { return CurveEditor; }
+
 	void HandlePostChange();
 
 	/**	Spawns the tab with the curve table inside */
@@ -80,6 +83,9 @@ public:
 	/** Get the mode that we are displaying data in */
 	ECurveTableViewMode GetViewMode() const { return ViewMode; }
 
+	/** Rename a specific curve */
+	void HandleCurveRename(FCurveEditorTreeItemID& TreeID, FName& OldCurveName, FName& NewCurveName);
+
 protected:
 
 	/** Handles setting up slate for the curve table editor */
@@ -88,6 +94,9 @@ protected:
 	/** Add extra menu items */
 	void ExtendMenu();
 
+	/** Add extra menu items */
+	void ExtendToolbar();
+
 	/** Bind commands to delegates */
 	void BindCommands();
 
@@ -119,7 +128,7 @@ protected:
 	FReply OnAddCurveClicked();
 
 	/** Callback For SimpleCurves, add a new Key/Column */
-	FReply OnAddNewKeyColumn();
+	void OnAddNewKeyColumn();
 
 	/* Adds new key for all (Simple) curves in the table at given time */
 	void AddNewKeyColumn(float NewKeyTime);
@@ -130,8 +139,11 @@ protected:
 	/** Get whether the curve view checkbox should be toggled on */
 	bool IsCurveViewChecked() const;
 
+	/** Invoke UI for Renaming a Curve */
+	void OnRequestCurveRename(FCurveEditorTreeItemID TreeItemId);
+
 	virtual bool ShouldCreateDefaultStandaloneMenu() const { return true; }
-	virtual bool ShouldCreateDefaultToolbar() const { return false; }
+	virtual bool ShouldCreateDefaultToolbar() const { return true ; }
 
 	/** Array of the columns that are available for editing */
 	TArray<FCurveTableEditorColumnHeaderDataPtr> AvailableColumns;
@@ -145,6 +157,9 @@ protected:
 	/** Menu extender */
 	TSharedPtr<FExtender> MenuExtender;
 
+	/** Menu extender */
+	TSharedPtr<FExtender> ToolbarExtender;
+
 	/**	The tab id for the curve table tab */
 	static const FName CurveTableTabId;
 
diff --git a/Engine/Source/Editor/CurveTableEditor/Private/CurveTableEditorCommands.cpp b/Engine/Source/Editor/CurveTableEditor/Private/CurveTableEditorCommands.cpp
--- a/Engine/Source/Editor/CurveTableEditor/Private/CurveTableEditorCommands.cpp
+++ b/Engine/Source/Editor/CurveTableEditor/Private/CurveTableEditorCommands.cpp
@@ -7,6 +7,7 @@
 void FCurveTableEditorCommands::RegisterCommands()
 {
 	UI_COMMAND(CurveViewToggle, "Curve View", "Changes the view of the curve table from grid to curve view.", EUserInterfaceActionType::ToggleButton, FInputChord());
+	UI_COMMAND(AppendKeyColumn, "Append Key Column", "Append a new column to the curve table.\nEvery Curve or Table Row will have a new key appended.", EUserInterfaceActionType::Button, FInputChord());
 }
 
 #undef LOCTEXT_NAMESPACE
diff --git a/Engine/Source/Editor/CurveTableEditor/Private/CurveTableEditorCommands.h b/Engine/Source/Editor/CurveTableEditor/Private/CurveTableEditorCommands.h
--- a/Engine/Source/Editor/CurveTableEditor/Private/CurveTableEditorCommands.h
+++ b/Engine/Source/Editor/CurveTableEditor/Private/CurveTableEditorCommands.h
@@ -20,4 +20,5 @@ public:
 
 	// Command to allow users to toggle the view mode
 	TSharedPtr<FUICommandInfo> CurveViewToggle;
+	TSharedPtr<FUICommandInfo> AppendKeyColumn;
 };
diff --git a/Engine/Source/Runtime/Engine/Classes/Engine/CurveTable.h b/Engine/Source/Runtime/Engine/Classes/Engine/CurveTable.h
--- a/Engine/Source/Runtime/Engine/Classes/Engine/CurveTable.h
+++ b/Engine/Source/Runtime/Engine/Classes/Engine/CurveTable.h
@@ -58,6 +58,9 @@ class UCurveTable
 	ENGINE_API FRichCurve& AddRichCurve(FName RowName);
 	ENGINE_API FSimpleCurve& AddSimpleCurve(FName RowName);
 
+	/** Move the curve to another FName in the table */
+	ENGINE_API void RenameRow(FName& CurveName, FName& NewCurveName);
+
 protected:
 	/** 
 	 * Map of name of row to row data structure. 
@@ -229,6 +232,7 @@ protected:
 
 	ENGINE_API static int32 GlobalCachedCurveID;
 
+
 private:
 
 	/** A multicast delegate that is called any time the curve table changes. */
diff --git a/Engine/Source/Runtime/Engine/Private/CurveTable.cpp b/Engine/Source/Runtime/Engine/Private/CurveTable.cpp
--- a/Engine/Source/Runtime/Engine/Private/CurveTable.cpp
+++ b/Engine/Source/Runtime/Engine/Private/CurveTable.cpp
@@ -577,6 +577,17 @@ FSimpleCurve& UCurveTable::AddSimpleCurve(FName RowName)
 	return *Result;
 }
 
+void UCurveTable::RenameRow(FName& CurveName, FName& NewCurveName)
+{
+	if (CurveName != NewCurveName && !RowMap.Contains(NewCurveName))
+	{
+		if (FRealCurve** Curve = RowMap.Find(CurveName))
+		{
+			RowMap.Add(NewCurveName, *Curve);
+			RowMap.Remove(CurveName);
+		}
+	}
+}
 
 /** */
 void GetCurveValues(const TArray<const TCHAR*>& Cells, TArray<float>& Values)
