CurveTableEditor - Remove Curves & Adds Right-Click Context Menu to CurveRows List

diff --git a/Engine/Source/Editor/CurveEditor/Private/Tree/SCurveEditorTree.cpp b/Engine/Source/Editor/CurveEditor/Private/Tree/SCurveEditorTree.cpp
--- a/Engine/Source/Editor/CurveEditor/Private/Tree/SCurveEditorTree.cpp
+++ b/Engine/Source/Editor/CurveEditor/Private/Tree/SCurveEditorTree.cpp
@@ -101,6 +101,7 @@ void SCurveEditorTree::Construct(const FArguments& InArgs, TSharedPtr<FCurveEdit
 		.OnGenerateRow(this, &SCurveEditorTree::GenerateRow)
 		.OnSetExpansionRecursive(this, &SCurveEditorTree::SetItemExpansionRecursive)
 		.OnMouseButtonDoubleClick(InArgs._OnMouseButtonDoubleClick)
+		.OnContextMenuOpening(InArgs._OnContextMenuOpening)
 		.OnTreeViewScrolled(InArgs._OnTreeViewScrolled)
 		.OnSelectionChanged_Lambda(
 			[this](TListTypeTraits<FCurveEditorTreeItemID>::NullableType InItemID, ESelectInfo::Type Type)
diff --git a/Engine/Source/Editor/CurveEditor/Public/Tree/SCurveEditorTree.h b/Engine/Source/Editor/CurveEditor/Public/Tree/SCurveEditorTree.h
--- a/Engine/Source/Editor/CurveEditor/Public/Tree/SCurveEditorTree.h
+++ b/Engine/Source/Editor/CurveEditor/Public/Tree/SCurveEditorTree.h
@@ -23,6 +23,7 @@ public:
 		SLATE_ARGUMENT(float, SelectColumnWidth)
 		SLATE_EVENT(FOnMouseButtonDoubleClick, OnMouseButtonDoubleClick)
 		SLATE_EVENT(FOnTableViewScrolled, OnTreeViewScrolled)
+		SLATE_EVENT(FOnContextMenuOpening, OnContextMenuOpening)
 	SLATE_END_ARGS()
 
 	void Construct(const FArguments& InArgs, TSharedPtr<FCurveEditor> InCurveEditor);
diff --git a/Engine/Source/Editor/CurveTableEditor/Private/CurveTableEditor.cpp b/Engine/Source/Editor/CurveTableEditor/Private/CurveTableEditor.cpp
--- a/Engine/Source/Editor/CurveTableEditor/Private/CurveTableEditor.cpp
+++ b/Engine/Source/Editor/CurveTableEditor/Private/CurveTableEditor.cpp
@@ -386,6 +386,18 @@ void FCurveTableEditor::BindCommands()
 		FCurveTableEditorCommands::Get().AppendKeyColumn,
 		FExecuteAction::CreateSP(this, &FCurveTableEditor::OnAddNewKeyColumn)
 	);
+
+	ToolkitCommands->MapAction(
+		FCurveTableEditorCommands::Get().RenameSelectedCurve,
+		FExecuteAction::CreateSP(this, &FCurveTableEditor::OnRenameCurve)
+	);
+
+
+	ToolkitCommands->MapAction(
+		FCurveTableEditorCommands::Get().DeleteSelectedCurves,
+		FExecuteAction::CreateSP(this, &FCurveTableEditor::OnDeleteCurves)
+	);
+
 }
 
 void FCurveTableEditor::ExtendMenu()
@@ -552,7 +564,8 @@ TSharedRef<SDockTab> FCurveTableEditor::SpawnTab_CurveTable( const FSpawnTabArgs
 
 	CurveEditorTree = SNew(SCurveEditorTree, CurveEditor.ToSharedRef())
 		.OnTreeViewScrolled(this, &FCurveTableEditor::OnCurveTreeViewScrolled)
-		.OnMouseButtonDoubleClick(this, &FCurveTableEditor::OnRequestCurveRename);
+		.OnMouseButtonDoubleClick(this, &FCurveTableEditor::OnRequestCurveRename)
+		.OnContextMenuOpening(this, &FCurveTableEditor::OnOpenCurveMenu);
 
 	TSharedRef<SCurveEditorPanel> CurveEditorPanel = SNew(SCurveEditorPanel, CurveEditor.ToSharedRef());
 
@@ -1057,4 +1070,77 @@ void FCurveTableEditor::HandleCurveRename(FCurveEditorTreeItemID& TreeID, FName&
 
 }
 
+void FCurveTableEditor::OnRenameCurve()
+{
+	const TMap<FCurveEditorTreeItemID, ECurveEditorTreeSelectionState>& SelectedRows = CurveEditor->GetTreeSelection();
+	if (SelectedRows.Num() == 1)
+	{
+		for (auto Item : SelectedRows)
+		{
+			OnRequestCurveRename(Item.Key);
+		}		
+	}
+}
+
+void FCurveTableEditor::OnDeleteCurves()
+{
+	UCurveTable* Table = Cast<UCurveTable>(GetEditingObject());
+	check(Table != nullptr);
+
+	const TMap<FCurveEditorTreeItemID, ECurveEditorTreeSelectionState>& SelectedRows = CurveEditor->GetTreeSelection();
+
+	if (SelectedRows.Num() >= 1)
+	{
+		FScopedTransaction Transaction(LOCTEXT("DeleteCurveRow", "Delete Curve Rows"));
+		Table->SetFlags(RF_Transactional);
+		Table->Modify();
+
+		for (auto Item : SelectedRows)
+		{
+			CurveEditor->RemoveTreeItem(Item.Key);
+
+			FName& CurveName = RowIDMap[Item.Key];
+
+			Table->DeleteRow(CurveName);
+
+			RowIDMap.Remove(Item.Key);
+		}
+
+		FPropertyChangedEvent PropertyChangeStruct(nullptr, EPropertyChangeType::ValueSet);
+		Table->PostEditChangeProperty(PropertyChangeStruct);
+	}
+}
+
+TSharedPtr<SWidget> FCurveTableEditor::OnOpenCurveMenu()
+{
+	int32 SelectedRowCount = CurveEditor->GetTreeSelection().Num();
+	if (SelectedRowCount > 0)
+	{
+		FMenuBuilder MenuBuilder(true /*auto close*/, ToolkitCommands);
+		MenuBuilder.BeginSection("Edit");
+		if (SelectedRowCount == 1)
+		{
+			MenuBuilder.AddMenuEntry(
+				FCurveTableEditorCommands::Get().RenameSelectedCurve,
+				NAME_None,
+				TAttribute<FText>(),
+				TAttribute<FText>(),
+				FSlateIcon(FAppStyle::GetAppStyleSetName(), "Icons.Edit")
+			);
+		}
+		MenuBuilder.AddMenuEntry(
+			FCurveTableEditorCommands::Get().DeleteSelectedCurves,
+			NAME_None,
+			TAttribute<FText>(),
+			TAttribute<FText>(),
+			FSlateIcon(FAppStyle::GetAppStyleSetName(), "Icons.Delete")
+		);
+		MenuBuilder.EndSection();
+
+		return MenuBuilder.MakeWidget();
+	}
+
+	return SNullWidget::NullWidget;
+}
+
 #undef LOCTEXT_NAMESPACE
\ No newline at end of file
diff --git a/Engine/Source/Editor/CurveTableEditor/Private/CurveTableEditor.h b/Engine/Source/Editor/CurveTableEditor/Private/CurveTableEditor.h
--- a/Engine/Source/Editor/CurveTableEditor/Private/CurveTableEditor.h
+++ b/Engine/Source/Editor/CurveTableEditor/Private/CurveTableEditor.h
@@ -86,6 +86,15 @@ public:
 	/** Rename a specific curve */
 	void HandleCurveRename(FCurveEditorTreeItemID& TreeID, FName& OldCurveName, FName& NewCurveName);
 
+	/** Pop open context Menu */
+	TSharedPtr<SWidget> OnOpenCurveMenu();
+
+	/** Callback for R-Click Menu Delete Curves */
+	void OnDeleteCurves();
+
+	/** Callback for R-Click Menu Rename Selected Curve */
+	void OnRenameCurve();
+
 protected:
 
 	/** Handles setting up slate for the curve table editor */
diff --git a/Engine/Source/Editor/CurveTableEditor/Private/CurveTableEditorCommands.cpp b/Engine/Source/Editor/CurveTableEditor/Private/CurveTableEditorCommands.cpp
--- a/Engine/Source/Editor/CurveTableEditor/Private/CurveTableEditorCommands.cpp
+++ b/Engine/Source/Editor/CurveTableEditor/Private/CurveTableEditorCommands.cpp
@@ -8,6 +8,8 @@ void FCurveTableEditorCommands::RegisterCommands()
 {
 	UI_COMMAND(CurveViewToggle, "Curve View", "Changes the view of the curve table from grid to curve view.", EUserInterfaceActionType::ToggleButton, FInputChord());
 	UI_COMMAND(AppendKeyColumn, "Append Key Column", "Append a new column to the curve table.\nEvery Curve or Table Row will have a new key appended.", EUserInterfaceActionType::Button, FInputChord());
+	UI_COMMAND(DeleteSelectedCurves, "Delete Curves", "Deleted the selected curve rows.", EUserInterfaceActionType::None, FInputChord());
+	UI_COMMAND(RenameSelectedCurve, "Rename Curve", "Rename the selected curve row.", EUserInterfaceActionType::None, FInputChord());
 }
 
 #undef LOCTEXT_NAMESPACE
diff --git a/Engine/Source/Editor/CurveTableEditor/Private/CurveTableEditorCommands.h b/Engine/Source/Editor/CurveTableEditor/Private/CurveTableEditorCommands.h
--- a/Engine/Source/Editor/CurveTableEditor/Private/CurveTableEditorCommands.h
+++ b/Engine/Source/Editor/CurveTableEditor/Private/CurveTableEditorCommands.h
@@ -21,4 +21,6 @@ public:
 	// Command to allow users to toggle the view mode
 	TSharedPtr<FUICommandInfo> CurveViewToggle;
 	TSharedPtr<FUICommandInfo> AppendKeyColumn;
+	TSharedPtr<FUICommandInfo> RenameSelectedCurve;
+	TSharedPtr<FUICommandInfo> DeleteSelectedCurves;
 };
diff --git a/Engine/Source/Runtime/Engine/Classes/Engine/CurveTable.h b/Engine/Source/Runtime/Engine/Classes/Engine/CurveTable.h
--- a/Engine/Source/Runtime/Engine/Classes/Engine/CurveTable.h
+++ b/Engine/Source/Runtime/Engine/Classes/Engine/CurveTable.h
@@ -61,6 +61,9 @@ class UCurveTable
 	/** Move the curve to another FName in the table */
 	ENGINE_API void RenameRow(FName& CurveName, FName& NewCurveName);
 
+	/** Remove a curve row from the table.  Note the associated curve will be deleted. */
+	ENGINE_API void DeleteRow(FName& CurveName);
+
 protected:
 	/** 
 	 * Map of name of row to row data structure. 
diff --git a/Engine/Source/Runtime/Engine/Private/CurveTable.cpp b/Engine/Source/Runtime/Engine/Private/CurveTable.cpp
--- a/Engine/Source/Runtime/Engine/Private/CurveTable.cpp
+++ b/Engine/Source/Runtime/Engine/Private/CurveTable.cpp
@@ -589,6 +589,19 @@ void UCurveTable::RenameRow(FName& CurveName, FName& NewCurveName)
 	}
 }
 
+void UCurveTable::DeleteRow(FName& CurveName)
+{
+	if (RowMap.Contains(CurveName))
+	{
+		FRealCurve** Curve = RowMap.Find(CurveName);
+		RowMap.Remove(CurveName);
+		if (Curve != nullptr && *Curve != nullptr)
+		{
+			delete *Curve;
+		}
+	}	
+}
+
 /** */
 void GetCurveValues(const TArray<const TCHAR*>& Cells, TArray<float>& Values)
 {
