Curve Table Editor - Adds ability to Retime and Delete Key Columns

diff --git a/Engine/Source/Editor/CurveTableEditor/Private/CurveTableEditor.cpp b/Engine/Source/Editor/CurveTableEditor/Private/CurveTableEditor.cpp
--- a/Engine/Source/Editor/CurveTableEditor/Private/CurveTableEditor.cpp
+++ b/Engine/Source/Editor/CurveTableEditor/Private/CurveTableEditor.cpp
@@ -816,12 +816,7 @@ void FCurveTableEditor::RefreshCachedCurveTable()
 
 			AvailableColumns.Add(CachedColumnData);
 
-			ColumnNamesHeaderRow->AddColumn(
-				SHeaderRow::Column(CachedColumnData->ColumnId)
-				.DefaultLabel(CachedColumnData->DisplayName)
-				.FixedWidth(CachedColumnData->DesiredColumnWidth + 50)
-				.HAlignHeader(HAlign_Center)
-			);
+			ColumnNamesHeaderRow->AddColumn( GenerateHeaderColumnForKey(CachedColumnData) );
 		}
 
 		// Setup the CurveEditorTree 
@@ -1006,7 +1001,7 @@ void FCurveTableEditor::AddNewKeyColumn(float NewKeyTime)
 		Table->Modify();	
 
 		// Make sure we don't already have a key at this time
-
+		
 		// 1. Add new keys to every curve
 		for (const TPair<FName, FRealCurve*>& CurveRow : Table->GetRowMap())
 		{
@@ -1019,21 +1014,18 @@ void FCurveTableEditor::AddNewKeyColumn(float NewKeyTime)
 		const FText ColumnText = FText::AsNumber(NewKeyTime);
 		ColumnData->ColumnId = *ColumnText.ToString();
 		ColumnData->DisplayName = ColumnText;
-		// ColumnData->DesiredColumnWidth = FontMeasure->Measure(ColumnData->DisplayName, CellTextStyle.Font).X + CellPadding;
 		ColumnData->KeyTime = NewKeyTime;
 
+		TSharedRef<FSlateFontMeasure> FontMeasure = FSlateApplication::Get().GetRenderer()->GetFontMeasureService();
+		const FTextBlockStyle& CellTextStyle = FEditorStyle::GetWidgetStyle<FTextBlockStyle>("DataTableEditor.CellText");
+		ColumnData->DesiredColumnWidth = FontMeasure->Measure(ColumnData->DisplayName, CellTextStyle.Font).X + 10.f;
+
 		AvailableColumns.Add(ColumnData);
 
 		// 3. Let the CurveTreeItems know they need to recache
 		OnColumnsChanged.Broadcast();
 
-		// Add the column to the TableView Header Row
-		ColumnNamesHeaderRow->AddColumn(
-			SHeaderRow::Column(ColumnData->ColumnId)
-			.DefaultLabel(ColumnData->DisplayName)
-			.FixedWidth(ColumnData->DesiredColumnWidth + 50)
-			.HAlignHeader(HAlign_Center)
-		);	
+		ColumnNamesHeaderRow->AddColumn( GenerateHeaderColumnForKey(ColumnData) );
 	}
 }
 
@@ -1143,4 +1135,200 @@ TSharedPtr<SWidget> FCurveTableEditor::OnOpenCurveMenu()
 	return SNullWidget::NullWidget;
 }
 
+void FCurveTableEditor::OnDeleteKeyColumn(float KeyTime)
+{
+	UCurveTable* Table = Cast<UCurveTable>(GetEditingObject());
+	check(Table != nullptr);
+
+	if (!Table->HasRichCurves())
+	{
+		// First find the column data associated with the original keytime
+		int FoundIndex = -1;
+		for (int i = 0; i < AvailableColumns.Num(); i++)
+		{
+			if (FMath::IsNearlyEqual(KeyTime, AvailableColumns[i]->KeyTime, KINDA_SMALL_NUMBER))
+			{
+				FoundIndex = i;
+				break;
+			}
+		}
+		if (FoundIndex < 0)
+		{
+			return;
+		}
+
+		FCurveTableEditorColumnHeaderDataPtr ColumnData = AvailableColumns[FoundIndex];
+		if (ColumnData.IsValid())
+		{
+
+			// Remove the column from the ui 
+			AvailableColumns.RemoveAt(FoundIndex);
+			ColumnNamesHeaderRow->RemoveColumn(ColumnData->ColumnId);
+
+			// Remove the keys from all curve rows is the data table
+			FScopedTransaction Transaction(LOCTEXT("DeleteKeyColumn", "Delete Key Column"));
+			Table->Modify();	
+
+			for (const TPair<FName, FRealCurve*>& CurveRow : Table->GetRowMap())
+			{
+				FRealCurve* Curve = CurveRow.Value;
+				FKeyHandle KeyHandle = Curve->FindKey(KeyTime);
+				if (KeyHandle != FKeyHandle::Invalid())
+				{
+					Curve->DeleteKey(KeyHandle);
+				}
+			}
+
+			FPropertyChangedEvent PropertyChangeStruct(nullptr, EPropertyChangeType::ValueSet);
+			Table->PostEditChangeProperty(PropertyChangeStruct);
+
+			// Let the CurveTreeItems (row ui) know they need to recache
+			OnColumnsChanged.Broadcast();
+		}
+
+	}
+}
+
+bool FCurveTableEditor::VerifyValidRetime(const FText& InText, FText& OutErrorMessage, float OriginalTime)
+{
+	if (!InText.IsNumeric())
+	{
+		OutErrorMessage = LOCTEXT("KeysMustBeNumeric", "Key Times must be numeric.");
+		return false;
+	}
+
+	float NewTime = 0.0f;
+	LexFromString(NewTime, *InText.ToString());
+
+	// do we already have a column with this time? 
+	for (auto Col : AvailableColumns)
+	{
+		if (FMath::IsNearlyEqual(NewTime, Col->KeyTime, KINDA_SMALL_NUMBER))
+		{
+			OutErrorMessage = LOCTEXT("KeyAlreadyExists", "Key times must be unique!");
+			return false;
+		}
+	}
+	return true;
+}
+
+void FCurveTableEditor::HandleRetimeCommitted(const FText& InText, ETextCommit::Type CommitInfo, float OriginalKeyTime)
+{
+
+	// First find the column data associated with the original keytime
+	int FoundIndex = -1;
+	for (int i = 0; i < AvailableColumns.Num(); i++)
+	{
+		if (FMath::IsNearlyEqual(OriginalKeyTime, AvailableColumns[i]->KeyTime, KINDA_SMALL_NUMBER))
+		{
+			FoundIndex = i;
+			break;
+		}
+	}
+	if (FoundIndex < 0)
+	{
+		return;
+	}
+
+	FCurveTableEditorColumnHeaderDataPtr CachedColumnData = AvailableColumns[FoundIndex];
+	if (CachedColumnData.IsValid())
+	{
+		// 1. Remove the UI associated with this column (ColumnData and the SHeaderRow::FColumn)
+		AvailableColumns.RemoveAt(FoundIndex);
+		ColumnNamesHeaderRow->RemoveColumn(CachedColumnData->ColumnId);
+
+		float NewTime = 0.0f;
+		LexFromString(NewTime, *InText.ToString());
+	
+		// 2. Adjust the key times for each of the curve table rows
+		UCurveTable* Table = Cast<UCurveTable>(GetEditingObject());
+		check(Table != nullptr);
+
+		FScopedTransaction Transaction(LOCTEXT("RetimeKeyColumn", "Retime Key Column"));
+		Table->Modify();	
+
+		for (const TPair<FName, FRealCurve*>& CurveRow : Table->GetRowMap())
+		{
+			FRealCurve* Curve = CurveRow.Value;
+			FKeyHandle KeyHandle = Curve->FindKey(OriginalKeyTime);
+			if (KeyHandle != FKeyHandle::Invalid())
+			{
+				Curve->SetKeyTime(KeyHandle, NewTime);
+			}
+		}
+
+		FPropertyChangedEvent PropertyChangeStruct(nullptr, EPropertyChangeType::ValueSet);
+		Table->PostEditChangeProperty(PropertyChangeStruct);
+
+
+		// 3. Update the ColumnData and re-insert the ColumnData and SHeaderRow::FColumn into the 
+		// correct places in order of the key times
+		int NewIndex = 0;
+		while (NewIndex < AvailableColumns.Num() && NewTime > AvailableColumns[NewIndex]->KeyTime )
+		{
+			NewIndex++;
+		}
+
+		const FText ColumnText = FText::AsNumber(NewTime);
+		CachedColumnData->ColumnId = *ColumnText.ToString();
+		CachedColumnData->DisplayName = ColumnText;
+		CachedColumnData->KeyTime = NewTime;
+
+		AvailableColumns.Insert(CachedColumnData, NewIndex);
+
+		// Let the CurveTreeItems know they need to recache
+		// note we do this before adding the column to the header so the rows already have their 
+		// data in place and are prepared to draw
+		OnColumnsChanged.Broadcast();
+
+		ColumnNamesHeaderRow->InsertColumn( GenerateHeaderColumnForKey(CachedColumnData), NewIndex)	;
+	}
+}
+
+SHeaderRow::FColumn::FArguments FCurveTableEditor::GenerateHeaderColumnForKey(FCurveTableEditorColumnHeaderDataPtr ColumnData)
+{
+	TSharedRef<SInlineEditableTextBlock> KeyTimeWidget = SNew(SInlineEditableTextBlock)
+	.Text(ColumnData->DisplayName)
+	.Justification(ETextJustify::Center)
+	.ColorAndOpacity(FSlateColor::UseForeground())
+	.OnTextCommitted(this, &FCurveTableEditor::HandleRetimeCommitted, ColumnData->KeyTime)
+	.OnVerifyTextChanged(this, &FCurveTableEditor::VerifyValidRetime, ColumnData->KeyTime);
+
+	// Create the Column Header's R-Click Menu
+	FMenuBuilder MenuBuilder(true /*Auto close*/, ToolkitCommands);
+	MenuBuilder.BeginSection("Edit");
+	MenuBuilder.AddMenuEntry(
+		FText::Format(LOCTEXT("RetimeKeysColumn", "Retime Keys at  {0}"), FText::AsNumber(ColumnData->KeyTime)),
+		FText::Format(LOCTEXT("RetimeKeysColumn_Tooltip", "Retimes this column and all keys at  {0}"), FText::AsNumber(ColumnData->KeyTime)),
+		FSlateIcon(FAppStyle::GetAppStyleSetName(), "Icons.Edit"),
+		FUIAction(FExecuteAction::CreateSP(KeyTimeWidget, &SInlineEditableTextBlock::EnterEditingMode))
+	);
+	MenuBuilder.AddMenuEntry(
+		FText::Format(LOCTEXT("DeleteKeysColumn", "Delete Keys at  {0}"), FText::AsNumber(ColumnData->KeyTime)),
+		FText::Format(LOCTEXT("DeleteKeysColumn_Tooltip", "Deletes this column and all keys at  {0}"), FText::AsNumber(ColumnData->KeyTime)),
+		FSlateIcon(FAppStyle::GetAppStyleSetName(), "Icons.Delete"),
+		FUIAction(FExecuteAction::CreateSP(this, &FCurveTableEditor::OnDeleteKeyColumn, ColumnData->KeyTime))
+	);
+	MenuBuilder.EndSection();
+
+	return SHeaderRow::Column(ColumnData->ColumnId)
+		.DefaultLabel(ColumnData->DisplayName)
+		.FixedWidth(ColumnData->DesiredColumnWidth + 40)
+		.HAlignHeader(HAlign_Fill)
+		.MenuContent()
+		[
+			MenuBuilder.MakeWidget()
+		]
+		.HeaderContent()
+		[
+			SNew(SBox)
+			.HeightOverride(22.f)
+			.HAlign(HAlign_Fill)
+			.VAlign(VAlign_Fill)
+			[
+				KeyTimeWidget
+			]
+		];
+}
+
 #undef LOCTEXT_NAMESPACE
\ No newline at end of file
diff --git a/Engine/Source/Editor/CurveTableEditor/Private/CurveTableEditor.h b/Engine/Source/Editor/CurveTableEditor/Private/CurveTableEditor.h
--- a/Engine/Source/Editor/CurveTableEditor/Private/CurveTableEditor.h
+++ b/Engine/Source/Editor/CurveTableEditor/Private/CurveTableEditor.h
@@ -10,6 +10,7 @@
 #include "ICurveTableEditor.h"
 #include "Widgets/Views/STableViewBase.h"
 #include "Widgets/Views/STableRow.h"
+#include "Widgets/Views/SHeaderRow.h"
 #include "Tree/ICurveEditorTreeItem.h"
 #include "CurveTableEditorHandle.h"
 #include "CurveTableEditorUtils.h"
@@ -95,6 +96,15 @@ public:
 	/** Callback for R-Click Menu Rename Selected Curve */
 	void OnRenameCurve();
 
+	/** Callback for R-Click Menu Delete Key Column */
+	void OnDeleteKeyColumn(float KeyTime);
+
+	/** Ensure that whatever time we are reetiming to is not redundant */
+	bool VerifyValidRetime(const FText& InText, FText& OutErrorMessage, float OriginalTime);
+
+	/** Process the retime from the header column text entry */
+	void HandleRetimeCommitted(const FText& InText, ETextCommit::Type CommitInfo, float OriginalKeyTime);
+
 protected:
 
 	/** Handles setting up slate for the curve table editor */
@@ -154,6 +164,8 @@ protected:
 	virtual bool ShouldCreateDefaultStandaloneMenu() const { return true; }
 	virtual bool ShouldCreateDefaultToolbar() const { return true ; }
 
+	SHeaderRow::FColumn::FArguments GenerateHeaderColumnForKey(FCurveTableEditorColumnHeaderDataPtr ColumnData);
+
 	/** Array of the columns that are available for editing */
 	TArray<FCurveTableEditorColumnHeaderDataPtr> AvailableColumns;
 
