CurveTableEditor -  Add Reimport Button directly on the toolbar

#pf 61ae6b4f8358693a22c3acc4

diff --git a/Engine/Source/Editor/CurveTableEditor/Private/CurveTableEditor.cpp b/Engine/Source/Editor/CurveTableEditor/Private/CurveTableEditor.cpp
--- a/Engine/Source/Editor/CurveTableEditor/Private/CurveTableEditor.cpp
+++ b/Engine/Source/Editor/CurveTableEditor/Private/CurveTableEditor.cpp
@@ -439,9 +439,16 @@ void FCurveTableEditor::ExtendToolbar()
 		GetToolkitCommands(),
 		FToolBarExtensionDelegate::CreateLambda([this](FToolBarBuilder& ParentToolbarBuilder)
 		{
-
 			ParentToolbarBuilder.BeginSection("CurveTable");
 
+			ParentToolbarBuilder.AddToolBarButton(
+				FUIAction(FExecuteAction::CreateSP(this, &FCurveTableEditor::Reimport_Execute, GetEditingObject())),
+				NAME_None,
+				FText::GetEmpty(),
+				LOCTEXT("Reimport_Tooltip", "Reimport the Curve Table from the source file.  All changes will be lost.  This action cannot be undone."),
+				FSlateIcon(FAppStyle::Get().GetStyleSetName(), "Icons.Toolbar.Import")
+			);
+
 			bool HasRichCurves = GetCurveTable()->HasRichCurves();
 			ParentToolbarBuilder.AddWidget(
 				SNew(SSegmentedControl<ECurveTableViewMode>)
@@ -483,7 +490,6 @@ void FCurveTableEditor::ExtendToolbar()
 	);
 
 	AddToolbarExtender(ToolbarExtender);
-
 }
 
 FName FCurveTableEditor::GetToolkitFName() const
diff --git a/Engine/Source/Editor/EditorStyle/Private/StarshipStyle.cpp b/Engine/Source/Editor/EditorStyle/Private/StarshipStyle.cpp
--- a/Engine/Source/Editor/EditorStyle/Private/StarshipStyle.cpp
+++ b/Engine/Source/Editor/EditorStyle/Private/StarshipStyle.cpp
@@ -470,6 +470,7 @@ void FStarshipEditorStyle::FStyle::SetupGeneralStyles()
 		Set("Icons.Toolbar.Stop", new IMAGE_BRUSH_SVG("Starship/MainToolbar/stop", Icon20x20));
 		Set("Icons.Toolbar.Settings", new CORE_IMAGE_BRUSH_SVG( "Starship/Common/Settings", Icon20x20));
 		Set("Icons.Toolbar.Details", new IMAGE_BRUSH_SVG("Starship/Common/Details", Icon16x16));
+		Set("Icons.Toolbar.Import", new CORE_IMAGE_BRUSH_SVG( "Starship/Common/import", Icon20x20) );
 	}
 
 	Set( "WarningStripe", new IMAGE_BRUSH( "Common/WarningStripe", FVector2D(20,6), FLinearColor::White, ESlateBrushTileType::Horizontal ) );
