Makes CurveTableEditor read-only for CompositeCurveTables

diff --git a/Engine/Source/Editor/CurveTableEditor/Private/CurveTableEditor.cpp b/Engine/Source/Editor/CurveTableEditor/Private/CurveTableEditor.cpp
--- a/Engine/Source/Editor/CurveTableEditor/Private/CurveTableEditor.cpp
+++ b/Engine/Source/Editor/CurveTableEditor/Private/CurveTableEditor.cpp
@@ -33,6 +33,7 @@
 
 #include "RealCurveModel.h"
 #include "RichCurveEditorModel.h"
+#include "Engine/CompositeCurveTable.h"
 
 #include "CurveTableEditorCommands.h"
 #include "Framework/MultiBox/MultiBoxBuilder.h"
@@ -400,6 +401,12 @@ void FCurveTableEditor::BindCommands()
 
 }
 
+bool FCurveTableEditor::IsReadOnly() const
+{
+	/* Currently, the only read-only tables are composite curve tables */
+	return GetCurveTable()->IsA<UCompositeCurveTable>();
+}
+
 void FCurveTableEditor::ExtendMenu()
 {
 	MenuExtender = MakeShareable(new FExtender);
@@ -551,6 +558,8 @@ TSharedRef<SDockTab> FCurveTableEditor::SpawnTab_CurveTable( const FSpawnTabArgs
 
 	bUpdatingTableViewSelection = false;
 
+	bool bTableIsReadOnly = IsReadOnly();
+
 	TSharedRef<SScrollBar> VerticalScrollBar = SNew(SScrollBar)
 		.Orientation(Orient_Vertical);
 
@@ -576,6 +585,8 @@ TSharedRef<SDockTab> FCurveTableEditor::SpawnTab_CurveTable( const FSpawnTabArgs
 	TSharedRef<SCurveEditorPanel> CurveEditorPanel = SNew(SCurveEditorPanel, CurveEditor.ToSharedRef());
 
 	TableView = SNew(SListView<FCurveEditorTreeItemID>)
+
+		.IsEnabled(!bTableIsReadOnly)
 		.ListItemsSource(&EmptyItems)
 		.OnListViewScrolled(this, &FCurveTableEditor::OnTableViewScrolled)
 		.HeaderRow(ColumnNamesHeaderRow)
@@ -633,6 +644,7 @@ TSharedRef<SDockTab> FCurveTableEditor::SpawnTab_CurveTable( const FSpawnTabArgs
 								.Icon(FAppStyle::Get().GetBrush("Icons.Plus"))
 								.Text(LOCTEXT("Curve", "Curve"))
 								.OnClicked(this, &FCurveTableEditor::OnAddCurveClicked)
+								.Visibility(bTableIsReadOnly ? EVisibility::Collapsed : EVisibility::Visible)
 							]
 
 							+SHorizontalBox::Slot()	
@@ -675,6 +687,7 @@ TSharedRef<SDockTab> FCurveTableEditor::SpawnTab_CurveTable( const FSpawnTabArgs
 					[
 						SNew(SBox)
 						.Visibility(this, &FCurveTableEditor::GetCurveViewControlsVisibility)
+						.IsEnabled(!bTableIsReadOnly)
 						[
 							CurveEditorPanel
 						]
@@ -916,7 +929,8 @@ TSharedRef<SWidget> FCurveTableEditor::MakeToolbar(TSharedRef<SCurveEditorPanel>
 	ToolBarBuilder.EndSection();
 	// We just use all of the extenders as our toolbar, we don't have a need to create a separate toolbar.
 
-	bool HasRichCurves = GetCurveTable()->HasRichCurves();
+	bool bHasRichCurves = GetCurveTable()->HasRichCurves();
+	bool bTableIsReadOnly = IsReadOnly();
 
 	return SNew(SHorizontalBox)
 
diff --git a/Engine/Source/Editor/CurveTableEditor/Private/CurveTableEditor.h b/Engine/Source/Editor/CurveTableEditor/Private/CurveTableEditor.h
--- a/Engine/Source/Editor/CurveTableEditor/Private/CurveTableEditor.h
+++ b/Engine/Source/Editor/CurveTableEditor/Private/CurveTableEditor.h
@@ -73,6 +73,9 @@ public:
 	/** Get the curve table being edited */
 	UCurveTable* GetCurveTable() const;
 
+	/** Determine if the table being viewed, can be edited or is read-only */
+	bool IsReadOnly() const;
+
 	/** Get the curve editor used for the Curve View*/
 	TSharedPtr<FCurveEditor> GetCurveEditor() const { return CurveEditor; }
 
