Edigrate of Curve Data Table to most recent feature set

- CurveTable Factory Options

- Fix Non-Unity Build Options

- Undo / Redo
- fixes memory leak

 - Constant Keys Support

- nonunity fix

- Switched to using an AssetEditorToolkit Toolbar
- Added support to rename curves

 - Remove Curves & Adds Right-Click Context Menu to CurveRows List

 - Adds ability to Retime and Delete Key Columns

-  Add Reimport Button directly on the toolbar

Curve Editor Stacked View - fixes label rendering

diff --git a/Engine/Source/Editor/CurveTableEditor/Private/CurveTableEditor.cpp b/Engine/Source/Editor/CurveTableEditor/Private/CurveTableEditor.cpp
--- a/Engine/Source/Editor/CurveTableEditor/Private/CurveTableEditor.cpp
+++ b/Engine/Source/Editor/CurveTableEditor/Private/CurveTableEditor.cpp
@@ -453,7 +453,7 @@ void FCurveTableEditor::ExtendToolbar()
 				NAME_None,
 				FText::GetEmpty(),
 				LOCTEXT("Reimport_Tooltip", "Reimport the Curve Table from the source file.  All changes will be lost.  This action cannot be undone."),
-				FSlateIcon(FAppStyle::GetAppStyleSetName(), "Icons.Toolbar.Import")
+				FSlateIcon(FAppStyle::Get().GetStyleSetName(), "Icons.Toolbar.Import")
 			);
 
 			bool HasRichCurves = GetCurveTable()->HasRichCurves();
@@ -470,28 +470,30 @@ void FCurveTableEditor::ExtendToolbar()
 			    .Icon(FAppStyle::Get().GetBrush("CurveTableEditor.TableView"))
 			);
 
-			if (InterpMode == RCIM_Constant)
+			if (!IsReadOnly())
 			{
-				ParentToolbarBuilder.AddToolBarButton(
-					FCurveTableEditorCommands::Get().AppendKeyColumn,
-					NAME_None, 
-					FText::GetEmpty(),
-					TAttribute<FText>(), 
-					FSlateIcon(FAppStyle::GetAppStyleSetName(), "Sequencer.KeySquare")
-				);
-			}
+				if (InterpMode == RCIM_Constant)
+				{
+					ParentToolbarBuilder.AddToolBarButton(
+						FCurveTableEditorCommands::Get().AppendKeyColumn,
+						NAME_None, 
+						FText::GetEmpty(),
+						TAttribute<FText>(), 
+						FSlateIcon(FAppStyle::Get().GetStyleSetName(), "Sequencer.KeySquare")
+					);
+				}
 
-			if (InterpMode == RCIM_Linear)
-			{
-				ParentToolbarBuilder.AddToolBarButton(
-					FCurveTableEditorCommands::Get().AppendKeyColumn,
-					NAME_None, 
-					FText::GetEmpty(),
-					TAttribute<FText>(), 
-					FSlateIcon(FAppStyle::GetAppStyleSetName(), "Sequencer.KeyTriangle")
-				);
+				if (InterpMode == RCIM_Linear)
+				{
+					ParentToolbarBuilder.AddToolBarButton(
+						FCurveTableEditorCommands::Get().AppendKeyColumn,
+						NAME_None, 
+						FText::GetEmpty(),
+						TAttribute<FText>(), 
+						FSlateIcon(FAppStyle::Get().GetStyleSetName(), "Sequencer.KeyTriangle")
+					);
+				}
 			}
-
 			ParentToolbarBuilder.EndSection();
 		})
 	);
@@ -1126,7 +1128,7 @@ void FCurveTableEditor::OnDeleteCurves()
 TSharedPtr<SWidget> FCurveTableEditor::OnOpenCurveMenu()
 {
 	int32 SelectedRowCount = CurveEditor->GetTreeSelection().Num();
-	if (SelectedRowCount > 0)
+	if (SelectedRowCount > 0 && !IsReadOnly())
 	{
 		FMenuBuilder MenuBuilder(true /*auto close*/, ToolkitCommands);
 		MenuBuilder.BeginSection("Edit");
diff --git a/Engine/Source/Editor/CurveTableEditor/Private/CurveTableEditor.h b/Engine/Source/Editor/CurveTableEditor/Private/CurveTableEditor.h
--- a/Engine/Source/Editor/CurveTableEditor/Private/CurveTableEditor.h
+++ b/Engine/Source/Editor/CurveTableEditor/Private/CurveTableEditor.h
@@ -73,6 +73,9 @@ public:
 	/** Get the curve table being edited */
 	UCurveTable* GetCurveTable() const;
 
+	/** Get the curve editor used for the Curve View*/
+	TSharedPtr<FCurveEditor> GetCurveEditor() const { return CurveEditor; }
+
 	/** Determine if the table being viewed, can be edited or is read-only */
 	bool IsReadOnly() const;
 
