Correcting merge error

diff --git a/Engine/Source/Editor/CurveTableEditor/Private/CurveTableEditor.h b/Engine/Source/Editor/CurveTableEditor/Private/CurveTableEditor.h
--- a/Engine/Source/Editor/CurveTableEditor/Private/CurveTableEditor.h
+++ b/Engine/Source/Editor/CurveTableEditor/Private/CurveTableEditor.h
@@ -73,9 +73,6 @@ public:
 	/** Get the curve table being edited */
 	UCurveTable* GetCurveTable() const;
 
-	/** Get the curve editor used for the Curve View*/
-	TSharedPtr<FCurveEditor> GetCurveEditor() const { return CurveEditor; }
-
 	/** Determine if the table being viewed, can be edited or is read-only */
 	bool IsReadOnly() const;
 
