CurveTableEditor nullptr checks

diff --git a/Engine/Source/Editor/CurveTableEditor/Private/CurveTableEditor.cpp b/Engine/Source/Editor/CurveTableEditor/Private/CurveTableEditor.cpp
--- a/Engine/Source/Editor/CurveTableEditor/Private/CurveTableEditor.cpp
+++ b/Engine/Source/Editor/CurveTableEditor/Private/CurveTableEditor.cpp
@@ -151,12 +151,21 @@ class FCurveTableEditorItem : public ICurveEditorTreeItem,  public TSharedFromTh
 
 			return SNew(SNumericEntryBox<float>)
 				.EditableTextBoxStyle( &FAppStyle::Get().GetWidgetStyle<FEditableTextBoxStyle>("CurveTableEditor.Cell.Text") )
-				.Value_Lambda([Curve, KeyHandle] () { return Curve->GetKeyValue(KeyHandle); })
+				.Value_Lambda([this, KeyHandle] () { 
+					if (FRealCurve* Curve = RowHandle.GetCurve())
+					{
+						return Curve->GetKeyValue(KeyHandle); 
+					}
+					return 0.0f;
+				})
 				.OnValueChanged_Lambda([this, KeyHandle] (float NewValue) 
 				{
-					FScopedTransaction Transaction(LOCTEXT("SetKeyValues", "Set Key Values"));
-					RowHandle.ModifyOwner();
-					RowHandle.GetCurve()->SetKeyValue(KeyHandle, NewValue);
+					if (FRealCurve* Curve = RowHandle.GetCurve())
+					{
+						FScopedTransaction Transaction(LOCTEXT("SetKeyValues", "Set Key Values"));
+						RowHandle.ModifyOwner();
+						Curve->SetKeyValue(KeyHandle, NewValue);
+					}
 				})
 				.Justification(ETextJustify::Right)
 			;
