Curve Table Editor: Add new "Add Key Column" icon for constant interpolation curve tables

diff --git a/Engine/Source/Editor/CurveTableEditor/Private/CurveTableEditor.cpp b/Engine/Source/Editor/CurveTableEditor/Private/CurveTableEditor.cpp
--- a/Engine/Source/Editor/CurveTableEditor/Private/CurveTableEditor.cpp
+++ b/Engine/Source/Editor/CurveTableEditor/Private/CurveTableEditor.cpp
@@ -526,28 +526,14 @@ void FCurveTableEditor::ExtendToolbar()
 
 			if (!IsReadOnly())
 			{
-				if (InterpMode == RCIM_Constant)
-				{
-					ParentToolbarBuilder.AddToolBarButton(
-						FCurveTableEditorCommands::Get().AppendKeyColumn,
-						NAME_None, 
-						FText::GetEmpty(),
-						TAttribute<FText>(), 
-						FSlateIcon(FAppStyle::Get().GetStyleSetName(), "Sequencer.KeySquare")
-					);
-				}
-
-				if (InterpMode == RCIM_Linear)
-				{
-					ParentToolbarBuilder.AddToolBarButton(
-						FCurveTableEditorCommands::Get().AppendKeyColumn,
-						NAME_None, 
-						FText::GetEmpty(),
-						TAttribute<FText>(), 
-						FSlateIcon(FAppStyle::Get().GetStyleSetName(), "Sequencer.KeyTriangle20")
-					);
-				}
+				ParentToolbarBuilder.AddToolBarButton(
+				FCurveTableEditorCommands::Get().AppendKeyColumn,
+				NAME_None, 
+				FText::GetEmpty(),
+				TAttribute<FText>(), 
+				FSlateIcon(FAppStyle::Get().GetStyleSetName(), "Sequencer.KeyTriangle20"));
 			}
+			
 			ParentToolbarBuilder.EndSection();
 		})
 	);
