Fix an issue with curve tables resetting data and preventing it from being editable in PIE

diff --git a/Engine/Source/Editor/CurveTableEditor/Private/CurveTableEditor.cpp b/Engine/Source/Editor/CurveTableEditor/Private/CurveTableEditor.cpp
--- a/Engine/Source/Editor/CurveTableEditor/Private/CurveTableEditor.cpp
+++ b/Engine/Source/Editor/CurveTableEditor/Private/CurveTableEditor.cpp
@@ -389,6 +389,11 @@ void FCurveTableEditor::UnregisterTabSpawners(const TSharedRef<class FTabManager
 FCurveTableEditor::~FCurveTableEditor()
 {
 	FReimportManager::Instance()->OnPostReimport().RemoveAll(this);
+	
+	if (UCurveTable* CurveTable = GetCurveTable())
+	{
+		CurveTable->OnCurveTableChanged().RemoveAll(this);
+	}
 }
 
 
@@ -403,7 +408,12 @@ void FCurveTableEditor::InitCurveTableEditor( const EToolkitMode::Type Mode, con
 	ExtendToolbar();
 	RegenerateMenusAndToolbars();
 
-	FReimportManager::Instance()->OnPostReimport().AddSP(this, &FCurveTableEditor::OnPostReimport);
+	FReimportManager::Instance()->OnPostReimport().AddSP(this, &FCurveTableEditor::OnPostReimport);	
+	
+	if (Table)
+	{
+		Table->OnCurveTableChanged().AddSP(this, &FCurveTableEditor::RefreshCachedCurveTable);
+	}
 
 	GEditor->RegisterForUndo(this);
 }
