Change the function called to refresh the curve tables, so the cached data is not destroyed on each modification / post edit change

#rnx

diff --git a/Engine/Source/Editor/CurveTableEditor/Private/CurveTableEditor.cpp b/Engine/Source/Editor/CurveTableEditor/Private/CurveTableEditor.cpp
--- a/Engine/Source/Editor/CurveTableEditor/Private/CurveTableEditor.cpp
+++ b/Engine/Source/Editor/CurveTableEditor/Private/CurveTableEditor.cpp
@@ -412,7 +412,7 @@ void FCurveTableEditor::InitCurveTableEditor( const EToolkitMode::Type Mode, con
 	
 	if (Table)
 	{
-		Table->OnCurveTableChanged().AddSP(this, &FCurveTableEditor::RefreshCachedCurveTable);
+		Table->OnCurveTableChanged().AddSP(this, &FCurveTableEditor::RefreshTableRows);
 	}
 
 	GEditor->RegisterForUndo(this);
