Curve Table Filtering;

Improving the curve table uasset filtering so we get more specific matches to the text queries. For example, typing "location" will filter all text with location in it but then typing "location.x" used to filter everything with location and x in it, now it'll filter things with location.x in it. So instead of filtering strings with x in it like "max", it now keeps both of the terms in consideration by finding x and then searching the hierarchy for location as well.

-- Code --
CurveTableEditor:
FCurveTableEditorItem::PassesFilter() - Updating the string search to not just be a Match() but to find each token in the string's hierarchy. Written by [REDACTED_PERSON].

#tests Editor - opened a curve table uasset and tried filtering by different variations of terms.

diff --git a/Engine/Source/Editor/CurveTableEditor/Private/CurveTableEditor.cpp b/Engine/Source/Editor/CurveTableEditor/Private/CurveTableEditor.cpp
--- a/Engine/Source/Editor/CurveTableEditor/Private/CurveTableEditor.cpp
+++ b/Engine/Source/Editor/CurveTableEditor/Private/CurveTableEditor.cpp
@@ -249,19 +249,25 @@ class FCurveTableEditorItem : public ICurveEditorTreeItem,  public TSharedFromTh
 	{
 		if (InFilter->GetType() == ECurveEditorTreeFilterType::Text)
 		{
+			FString     DisplayNameAsString = DisplayName.ToString();
+			FStringView RemainingStringView = DisplayNameAsString;
+
 			const FCurveEditorTreeTextFilter* Filter = static_cast<const FCurveEditorTreeTextFilter*>(InFilter);
 			for (const FCurveEditorTreeTextFilterTerm& Term : Filter->GetTerms())
 			{
-				for(const FCurveEditorTreeTextFilterToken& Token : Term.ChildToParentTokens)
+				for (const FCurveEditorTreeTextFilterToken& Token : Term.ChildToParentTokens)
 				{
-					if(Token.Match(*DisplayName.ToString()))
+					FStringView TokenView(Token.Token);
+					const int32 MatchPos = RemainingStringView.Find(TokenView, 0, ESearchCase::IgnoreCase);
+					if (MatchPos == INDEX_NONE)
 					{
-						return true;
+						return false;
 					}
+					RemainingStringView.LeftInline(MatchPos);
 				}
 			}
 
-			return false;
+			return true;
 		}
 
 		return false;
