Fix crash in composite curve editor if table view tab was closed, then asset editor re-opened and add (+) is clicked

diff --git a/Engine/Source/Editor/CurveTableEditor/Private/CurveTableEditor.cpp b/Engine/Source/Editor/CurveTableEditor/Private/CurveTableEditor.cpp
--- a/Engine/Source/Editor/CurveTableEditor/Private/CurveTableEditor.cpp
+++ b/Engine/Source/Editor/CurveTableEditor/Private/CurveTableEditor.cpp
@@ -750,12 +750,15 @@ TSharedRef<SDockTab> FCurveTableEditor::SpawnTab_CurveTable( const FSpawnTabArgs
 
 void FCurveTableEditor::RefreshTableRows()
 {
-	TableView->RequestListRefresh();
+	if (TableView)
+	{
+		TableView->RequestListRefresh();
+	}
 }
 
 void FCurveTableEditor::RefreshTableRowsSelection()
 {
-	if(bUpdatingTableViewSelection == false)
+	if(bUpdatingTableViewSelection == false && TableView)
 	{
 		TGuardValue<bool> SelectionGuard(bUpdatingTableViewSelection, true);
 
@@ -780,7 +783,7 @@ void FCurveTableEditor::RefreshTableRowsSelection()
 
 void FCurveTableEditor::OnTableViewSelectionChanged(FCurveEditorTreeItemID ItemID, ESelectInfo::Type)
 {
-	if (bUpdatingTableViewSelection == false)
+	if (bUpdatingTableViewSelection == false && TableView)
 	{
 		TGuardValue<bool> SelectionGuard(bUpdatingTableViewSelection, true);
 		CurveEditor->GetTree()->SetDirectSelection(TableView->GetSelectedItems(), CurveEditor.Get());
@@ -818,7 +821,10 @@ void FCurveTableEditor::RefreshCachedCurveTable()
 	// New Selection 
 	TArray<FCurveEditorTreeItemID> NewSelectedItems;
 
-	TableView->SetItemsSource(&EmptyItems);
+	if (TableView)
+	{
+		TableView->SetItemsSource(&EmptyItems);
+	}
 	
 	CurveEditor->RemoveAllTreeItems();
 
@@ -922,7 +928,10 @@ void FCurveTableEditor::RefreshCachedCurveTable()
 		}
 	}
 
-	TableView->SetItemsSource(&CurveEditorTree->GetSourceItems());
+	if (TableView)
+	{
+		TableView->SetItemsSource(&CurveEditorTree->GetSourceItems());
+	}
 
 	TGuardValue<bool> SelectionGuard(bUpdatingTableViewSelection, true);
 	CurveEditor->SetTreeSelection(MoveTemp(NewSelectedItems));
@@ -932,7 +941,10 @@ void FCurveTableEditor::RefreshCachedCurveTable()
 void FCurveTableEditor::OnCurveTreeViewScrolled(double InScrollOffset)
 {
 	// Synchronize the list views
-	TableView->SetScrollOffset(InScrollOffset);
+	if (TableView)
+	{
+		TableView->SetScrollOffset(InScrollOffset);
+	}
 }
 
 
