Follow up crash fix in composite curve editor if table view tab was closed, then asset editor re-opened and add (+) is clicked

diff --git a/Engine/Source/Editor/CurveTableEditor/Private/CurveTableEditor.cpp b/Engine/Source/Editor/CurveTableEditor/Private/CurveTableEditor.cpp
--- a/Engine/Source/Editor/CurveTableEditor/Private/CurveTableEditor.cpp
+++ b/Engine/Source/Editor/CurveTableEditor/Private/CurveTableEditor.cpp
@@ -750,7 +750,7 @@ TSharedRef<SDockTab> FCurveTableEditor::SpawnTab_CurveTable( const FSpawnTabArgs
 
 void FCurveTableEditor::RefreshTableRows()
 {
-	if (TableView)
+	if (TableView && CurveEditor)
 	{
 		TableView->RequestListRefresh();
 	}
@@ -758,7 +758,7 @@ void FCurveTableEditor::RefreshTableRows()
 
 void FCurveTableEditor::RefreshTableRowsSelection()
 {
-	if(bUpdatingTableViewSelection == false && TableView)
+	if(bUpdatingTableViewSelection == false && TableView && CurveEditor)
 	{
 		TGuardValue<bool> SelectionGuard(bUpdatingTableViewSelection, true);
 
@@ -783,7 +783,7 @@ void FCurveTableEditor::RefreshTableRowsSelection()
 
 void FCurveTableEditor::OnTableViewSelectionChanged(FCurveEditorTreeItemID ItemID, ESelectInfo::Type)
 {
-	if (bUpdatingTableViewSelection == false && TableView)
+	if (bUpdatingTableViewSelection == false && TableView && CurveEditor)
 	{
 		TGuardValue<bool> SelectionGuard(bUpdatingTableViewSelection, true);
 		CurveEditor->GetTree()->SetDirectSelection(TableView->GetSelectedItems(), CurveEditor.Get());
@@ -792,6 +792,11 @@ void FCurveTableEditor::OnTableViewSelectionChanged(FCurveEditorTreeItemID ItemI
 
 void FCurveTableEditor::RefreshCachedCurveTable()
 {
+	if (!CurveEditor)
+	{
+		return;
+	}
+
 	// This will trigger to remove any cached widgets in the TableView while we rebuild the model from the source CurveTable
 
 	const TSet<FCurveModelID>& Pinned = CurveEditor->GetPinnedCurves();
