Fix vertical scrolling in curve tables with no column or curve display.

[RN]

diff --git a/Engine/Source/Editor/CurveTableEditor/Private/CurveTableEditor.cpp b/Engine/Source/Editor/CurveTableEditor/Private/CurveTableEditor.cpp
--- a/Engine/Source/Editor/CurveTableEditor/Private/CurveTableEditor.cpp
+++ b/Engine/Source/Editor/CurveTableEditor/Private/CurveTableEditor.cpp
@@ -946,8 +946,10 @@ void FCurveTableEditor::RefreshCachedCurveTable()
 
 void FCurveTableEditor::OnCurveTreeViewScrolled(double InScrollOffset)
 {
-	// Synchronize the list views
-	if (TableView)
+	const UCurveTable* CurveTable = GetCurveTable();
+
+	// Synchronize the list views if the table has data in it
+	if (TableView && ((CurveTable && CurveTable->HasRichCurves()) || AvailableColumns.Num() > 0))
 	{
 		TableView->SetScrollOffset(InScrollOffset);
 	}
