Fixing ensure in CurveTableEditor when deleting a curve row.
Copying selected row map during deletion instead of taking reference which will change when rows are deleted.

#Synced-CL 50443327

diff --git a/Engine/Source/Editor/CurveTableEditor/Private/CurveTableEditor.cpp b/Engine/Source/Editor/CurveTableEditor/Private/CurveTableEditor.cpp
--- a/Engine/Source/Editor/CurveTableEditor/Private/CurveTableEditor.cpp
+++ b/Engine/Source/Editor/CurveTableEditor/Private/CurveTableEditor.cpp
@@ -1169,7 +1169,7 @@ void FCurveTableEditor::OnDeleteCurves()
 	UCurveTable* Table = Cast<UCurveTable>(GetEditingObject());
 	check(Table != nullptr);
 
-	const TMap<FCurveEditorTreeItemID, ECurveEditorTreeSelectionState>& SelectedRows = CurveEditor->GetTreeSelection();
+	const TMap<FCurveEditorTreeItemID, ECurveEditorTreeSelectionState> SelectedRows = CurveEditor->GetTreeSelection();
 
 	if (SelectedRows.Num() >= 1)
 	{
