Curve Table Editor: Fix dangling pointer in FRealCurveModel.

Instead of FRealCurveModel storing a pointer, it stores the handle and resolves it on demand.
This change adds an ensure to detect when the handle becomes invalid.

The UCurveTable::OnCurveTableChanged subscription in FCurveTableEditor has been adjusted to regenerate all FRealCurveModels. This should ensure that curve models do not become stale. We achieve this by deselecting from the tree view, which is hacky way of doing it but should be acceptable; the curve models are regenerated when the user makes a selection in the tree.

#Synced-CL 51626380

diff --git a/Engine/Source/Editor/CurveTableEditor/Private/CurveTableEditor.cpp b/Engine/Source/Editor/CurveTableEditor/Private/CurveTableEditor.cpp
--- a/Engine/Source/Editor/CurveTableEditor/Private/CurveTableEditor.cpp
+++ b/Engine/Source/Editor/CurveTableEditor/Private/CurveTableEditor.cpp
@@ -239,7 +239,7 @@ class FCurveTableEditorItem : public ICurveEditorTreeItem,  public TSharedFromTh
 			const UCurveTable* Table = RowHandle.CurveTable.Get();
 			UCurveTable* RawTable = const_cast<UCurveTable*>(Table);
 
-			TUniquePtr<FRealCurveModel> NewCurveModel = MakeUnique<FRealCurveModel>(RowHandle.GetCurve(), RawTable);
+			TUniquePtr<FRealCurveModel> NewCurveModel = MakeUnique<FRealCurveModel>(RowHandle, RawTable);
 			NewCurveModel->SetShortDisplayName(DisplayName);
 
 			OutCurveModels.Add(MoveTemp(NewCurveModel));
diff --git a/Engine/Source/Editor/CurveTableEditor/Private/RealCurveModel.cpp b/Engine/Source/Editor/CurveTableEditor/Private/RealCurveModel.cpp
--- a/Engine/Source/Editor/CurveTableEditor/Private/RealCurveModel.cpp
+++ b/Engine/Source/Editor/CurveTableEditor/Private/RealCurveModel.cpp
@@ -19,16 +19,14 @@
 
 class FCurveEditor;
 
-FRealCurveModel::FRealCurveModel(FRealCurve* InRealCurve, UObject* InOwner)
+FRealCurveModel::FRealCurveModel(FCurveTableEditorHandle InRealCurveHandle, UObject* InOwner)
 	: WeakOwner(InOwner)
-{
-
-	RealCurve = InRealCurve;
-}
+	, RealCurveHandle(MoveTemp(InRealCurveHandle))
+{}
 
 const void* FRealCurveModel::GetCurve() const
 {
-	return RealCurve;
+	return ResolveCurve();
 }
 
 void FRealCurveModel::Modify()
@@ -45,8 +43,7 @@ void FRealCurveModel::Modify()
 
 void FRealCurveModel::DrawCurve(const FCurveEditor& CurveEditor, const FCurveEditorScreenSpace& ScreenSpace, TArray<TTuple<double, double>>& InterpolatingPoints) const
 {
-
-	if (IsValid())
+	if (FRealCurve* RealCurve = ResolveCurve())
 	{
 
 		FKeyHandle LastKeyHandle = RealCurve->GetFirstKeyHandle();
@@ -85,7 +82,7 @@ void FRealCurveModel::DrawCurve(const FCurveEditor& CurveEditor, const FCurveEdi
 
 void FRealCurveModel::GetKeys(double MinTime, double MaxTime, double MinValue, double MaxValue, TArray<FKeyHandle>& OutKeyHandles) const
 {
-	if (IsValid())
+	if (FRealCurve* RealCurve = ResolveCurve())
 	{
 		for (auto It = RealCurve->GetKeyHandleIterator(); It; ++It)
 		{
@@ -110,7 +107,7 @@ void FRealCurveModel::RemoveKeys(TArrayView<const FKeyHandle> InKeys, double InC
 
 void FRealCurveModel::GetKeyPositions(TArrayView<const FKeyHandle> InKeys, TArrayView<FKeyPosition> OutKeyPositions) const
 {
-	if (IsValid())
+	if (FRealCurve* RealCurve = ResolveCurve())
 	{
 		for (int32 Index = 0; Index < InKeys.Num(); ++Index)
 		{
@@ -128,7 +125,7 @@ void FRealCurveModel::SetKeyPositions(TArrayView<const FKeyHandle> InKeys, TArra
 {
 	if (UObject* Owner = WeakOwner.Get())
 	{
-		if (IsValid())
+		if (FRealCurve* RealCurve = ResolveCurve())
 		{
 			Owner->Modify();
 
@@ -155,7 +152,7 @@ void FRealCurveModel::GetKeyDrawInfo(ECurvePointType PointType, const FKeyHandle
 		OutDrawInfo.Brush = FAppStyle::Get().GetBrush("GenericCurveEditor.TangentHandle");
 		OutDrawInfo.ScreenSize = FVector2D(9, 9);
 	}
-	else
+	else if (FRealCurve* RealCurve = ResolveCurve())
 	{
 		// All keys are the same size by default
 		OutDrawInfo.ScreenSize = FVector2D(11, 11);
@@ -181,7 +178,7 @@ void FRealCurveModel::GetKeyDrawInfo(ECurvePointType PointType, const FKeyHandle
 
 void FRealCurveModel::GetKeyAttributes(TArrayView<const FKeyHandle> InKeys, TArrayView<FKeyAttributes> OutAttributes) const
 {
-	if(IsValid())
+	if(FRealCurve* RealCurve = ResolveCurve())
 	{
 		for (int32 Index = 0; Index < InKeys.Num(); ++Index)
 		{
@@ -196,7 +193,7 @@ void FRealCurveModel::GetKeyAttributes(TArrayView<const FKeyHandle> InKeys, TArr
 
 void FRealCurveModel::GetTimeRange(double& MinTime, double& MaxTime) const
 {
-	if (IsValid())
+	if (FRealCurve* RealCurve = ResolveCurve())
 	{
 		float MinT, MaxT; 
 		RealCurve->GetTimeRange(MinT, MaxT);
@@ -207,7 +204,7 @@ void FRealCurveModel::GetTimeRange(double& MinTime, double& MaxTime) const
 
 void FRealCurveModel::GetValueRange(double& MinValue, double& MaxValue) const
 {
-	if (IsValid())
+	if (FRealCurve* RealCurve = ResolveCurve())
 	{
 		float MinV, MaxV;
 		RealCurve->GetValueRange(MinV, MaxV);
@@ -218,7 +215,7 @@ void FRealCurveModel::GetValueRange(double& MinValue, double& MaxValue) const
 
 int32 FRealCurveModel::GetNumKeys() const
 {
-	if (IsValid())
+	if (FRealCurve* RealCurve = ResolveCurve())
 	{
 		return RealCurve->GetNumKeys();
 	}
@@ -227,7 +224,7 @@ int32 FRealCurveModel::GetNumKeys() const
 
 void FRealCurveModel::GetNeighboringKeys(const FKeyHandle InKeyHandle, TOptional<FKeyHandle>& OutPreviousKeyHandle, TOptional<FKeyHandle>& OutNextKeyHandle) const
 {
-	if (IsValid())
+	if (FRealCurve* RealCurve = ResolveCurve())
 	{
 		OutPreviousKeyHandle = RealCurve->GetPreviousKey(InKeyHandle);
 		OutNextKeyHandle = RealCurve->GetNextKey(InKeyHandle);
@@ -236,7 +233,7 @@ void FRealCurveModel::GetNeighboringKeys(const FKeyHandle InKeyHandle, TOptional
 
 bool FRealCurveModel::Evaluate(double ProspectiveTime, double& OutValue) const
 {
-	if (IsValid())
+	if (FRealCurve* RealCurve = ResolveCurve())
 	{
 		OutValue = RealCurve->Eval(ProspectiveTime);
 		return true;
diff --git a/Engine/Source/Editor/CurveTableEditor/Private/RealCurveModel.h b/Engine/Source/Editor/CurveTableEditor/Private/RealCurveModel.h
--- a/Engine/Source/Editor/CurveTableEditor/Private/RealCurveModel.h
+++ b/Engine/Source/Editor/CurveTableEditor/Private/RealCurveModel.h
@@ -6,6 +6,7 @@
 #include "CoreTypes.h"
 #include "CurveEditorTypes.h"
 #include "CurveModel.h"
+#include "CurveTableEditorHandle.h"
 #include "Misc/OptionalFwd.h"
 #include "Templates/Tuple.h"
 #include "UObject/UnrealType.h"
@@ -30,7 +31,7 @@ class FRealCurveModel : public FCurveModel
   public:
 	static ECurveEditorViewID EventView;
 
-	FRealCurveModel(FRealCurve* InRealCurve, UObject* InOwner);
+	explicit FRealCurveModel(FCurveTableEditorHandle InRealCurveHandle, UObject* InOwner);
 
 	// FCurveModel interface
 	virtual const void* GetCurve() const override;
@@ -56,10 +57,17 @@ class FRealCurveModel : public FCurveModel
 
 	virtual bool Evaluate(double ProspectiveTime, double& OutValue) const override;
 
-	virtual bool IsValid() const { return RealCurve != nullptr; }
+	virtual bool IsValid() const { return RealCurveHandle.IsValid(); }
 
   private:
 
-  	TWeakObjectPtr<> WeakOwner;
-	FRealCurve* RealCurve;
+  	const TWeakObjectPtr<> WeakOwner;
+	const FCurveTableEditorHandle RealCurveHandle;
+	
+	FRealCurve* ResolveCurve(bool bEnsure = true) const
+	{
+		FRealCurve* Result = RealCurveHandle.GetCurve();
+		ensure(!bEnsure || Result);
+		return Result;
+	}
 };
\ No newline at end of file
